MovieScene: Fixed missing include

diff --git a/Engine/Plugins/MovieScene/ActorAnimationEditor/Source/ActorAnimationEditor/Private/ActorAnimationEditorPrivatePCH.h b/Engine/Plugins/MovieScene/ActorAnimationEditor/Source/ActorAnimationEditor/Private/ActorAnimationEditorPrivatePCH.h
--- a/Engine/Plugins/MovieScene/ActorAnimationEditor/Source/ActorAnimationEditor/Private/ActorAnimationEditorPrivatePCH.h
+++ b/Engine/Plugins/MovieScene/ActorAnimationEditor/Source/ActorAnimationEditor/Private/ActorAnimationEditorPrivatePCH.h
@@ -6,8 +6,9 @@
 /* Private dependencies
  *****************************************************************************/
 
-#include "ActorAnimation.h"
+#include "Core.h"
 #include "CoreUObject.h"
+#include "ActorAnimation.h"
 #include "AssetToolsModule.h"
 #include "EditorStyle.h"
 #include "Engine.h"
