ActorAnimation: Only enabled in content browser if running with "-sequencer"

diff --git a/Engine/Plugins/MovieScene/ActorAnimationEditor/Source/ActorAnimationEditor/Private/ActorAnimationFactoryNew.cpp b/Engine/Plugins/MovieScene/ActorAnimationEditor/Source/ActorAnimationEditor/Private/ActorAnimationFactoryNew.cpp
--- a/Engine/Plugins/MovieScene/ActorAnimationEditor/Source/ActorAnimationEditor/Private/ActorAnimationFactoryNew.cpp
+++ b/Engine/Plugins/MovieScene/ActorAnimationEditor/Source/ActorAnimationEditor/Private/ActorAnimationFactoryNew.cpp
@@ -39,6 +39,12 @@ UObject* UActorAnimationFactoryNew::FactoryCreateNew(UClass* Class, UObject* InP
 
 bool UActorAnimationFactoryNew::ShouldShowInNewMenu() const
 {
+	// @todo sequencer: Only allow users to create new MovieScenes if that feature is turned on globally.
+	if (!FParse::Param(FCommandLine::Get(), TEXT("Sequencer")))
+	{
+		return false;
+	}
+
 	return true;
 }
 
