ActorAnimation: Moved files into subfolders

diff --git a/Engine/Plugins/MovieScene/ActorAnimationEditor/Source/ActorAnimationEditor/ActorAnimationEditor.Build.cs b/Engine/Plugins/MovieScene/ActorAnimationEditor/Source/ActorAnimationEditor/ActorAnimationEditor.Build.cs
--- a/Engine/Plugins/MovieScene/ActorAnimationEditor/Source/ActorAnimationEditor/ActorAnimationEditor.Build.cs
+++ b/Engine/Plugins/MovieScene/ActorAnimationEditor/Source/ActorAnimationEditor/ActorAnimationEditor.Build.cs
@@ -33,5 +33,14 @@ public class ActorAnimationEditor : ModuleRules
 				"AssetTools",
 			}
 		);
+
+        PrivateIncludePaths.AddRange(
+            new string[] {
+				"ActorAnimationEditor/Private",
+				"ActorAnimationEditor/Private/AssetTools",
+				"ActorAnimationEditor/Private/Factories",
+				"ActorAnimationEditor/Private/Styles",
+			}
+        );
 	}
 }
diff --git a/Engine/Plugins/MovieScene/ActorAnimationEditor/Source/ActorAnimationEditor/Private/ActorAnimationActions.cpp b/Engine/Plugins/MovieScene/ActorAnimationEditor/Source/ActorAnimationEditor/Private/ActorAnimationActions.cpp
deleted file mode 100644
diff --git a/Engine/Plugins/MovieScene/ActorAnimationEditor/Source/ActorAnimationEditor/Private/ActorAnimationActions.h b/Engine/Plugins/MovieScene/ActorAnimationEditor/Source/ActorAnimationEditor/Private/ActorAnimationActions.h
deleted file mode 100644
diff --git a/Engine/Plugins/MovieScene/ActorAnimationEditor/Source/ActorAnimationEditor/Private/ActorAnimationFactoryNew.cpp b/Engine/Plugins/MovieScene/ActorAnimationEditor/Source/ActorAnimationEditor/Private/ActorAnimationFactoryNew.cpp
deleted file mode 100644
diff --git a/Engine/Plugins/MovieScene/ActorAnimationEditor/Source/ActorAnimationEditor/Private/ActorAnimationFactoryNew.h b/Engine/Plugins/MovieScene/ActorAnimationEditor/Source/ActorAnimationEditor/Private/ActorAnimationFactoryNew.h
deleted file mode 100644
diff --git a/Engine/Plugins/MovieScene/ActorAnimationEditor/Source/ActorAnimationEditor/Private/AssetTools/ActorAnimationActions.cpp b/Engine/Plugins/MovieScene/ActorAnimationEditor/Source/ActorAnimationEditor/Private/AssetTools/ActorAnimationActions.cpp
new file mode 100644
--- /dev/null
+++ b/Engine/Plugins/MovieScene/ActorAnimationEditor/Source/ActorAnimationEditor/Private/AssetTools/ActorAnimationActions.cpp
@@ -0,0 +1,72 @@
+// Copyright 1998-2015 Epic Games, Inc. All Rights Reserved.
+
+#include "ActorAnimationEditorPrivatePCH.h"
+#include "IAssetTools.h"
+#include "ISequencerModule.h"
+#include "Toolkits/IToolkit.h"
+
+
+#define LOCTEXT_NAMESPACE "AssetTypeActions"
+
+
+/* IAssetTypeActions interface
+ *****************************************************************************/
+
+uint32 FActorAnimationActions::GetCategories()
+{
+	return EAssetTypeCategories::Animation;
+}
+
+
+FText FActorAnimationActions::GetName() const
+{
+	return NSLOCTEXT("AssetTypeActions", "AssetTypeActions_ActorAnimation", "Actor Animation");
+}
+
+
+UClass* FActorAnimationActions::GetSupportedClass() const
+{
+	return UActorAnimation::StaticClass(); 
+}
+
+
+FColor FActorAnimationActions::GetTypeColor() const
+{
+	return FColor(200, 80, 80);
+}
+
+
+void FActorAnimationActions::OpenAssetEditor(const TArray<UObject*>& InObjects, TSharedPtr<IToolkitHost> EditWithinLevelEditor)
+{
+	// @todo sequencer: Only allow users to create new MovieScenes if that feature is turned on globally.
+	if (!FParse::Param(FCommandLine::Get(), TEXT("Sequencer")))
+	{
+		return;
+	}
+
+	EToolkitMode::Type Mode = EditWithinLevelEditor.IsValid()
+		? EToolkitMode::WorldCentric
+		: EToolkitMode::Standalone;
+
+	for (auto ObjIt = InObjects.CreateConstIterator(); ObjIt; ++ObjIt)
+	{
+		UActorAnimation* ActorAnimation = Cast<UActorAnimation>(*ObjIt);
+
+		if (ActorAnimation != nullptr)
+		{
+			FSequencerViewParams ViewParams(TEXT("ActorAnimationEditorViewParams"));
+			TSharedRef<FActorAnimationEditorToolkit> Toolkit = MakeShareable(new FActorAnimationEditorToolkit);
+			Toolkit->Initialize(Mode, ViewParams, EditWithinLevelEditor, ActorAnimation, true);
+		}
+	}
+}
+
+
+bool FActorAnimationActions::ShouldForceWorldCentric()
+{
+	// @todo sequencer: Hack to force world-centric mode for Sequencer
+	return true;
+}
+
+
+#undef LOCTEXT_NAMESPACE
\ No newline at end of file
diff --git a/Engine/Plugins/MovieScene/ActorAnimationEditor/Source/ActorAnimationEditor/Private/AssetTools/ActorAnimationActions.h b/Engine/Plugins/MovieScene/ActorAnimationEditor/Source/ActorAnimationEditor/Private/AssetTools/ActorAnimationActions.h
new file mode 100644
--- /dev/null
+++ b/Engine/Plugins/MovieScene/ActorAnimationEditor/Source/ActorAnimationEditor/Private/AssetTools/ActorAnimationActions.h
@@ -0,0 +1,24 @@
+// Copyright 1998-2015 Epic Games, Inc. All Rights Reserved.
+
+#pragma once
+
+#include "AssetTypeActions_Base.h"
+
+
+/**
+ * Implements actions for UActorAnimation assets.
+ */
+class FActorAnimationActions
+	: public FAssetTypeActions_Base
+{
+public:
+	
+	// IAssetTypeActions interface
+
+	virtual uint32 GetCategories() override;
+	virtual FText GetName() const override;
+	virtual UClass* GetSupportedClass() const override;
+	virtual FColor GetTypeColor() const override;
+	virtual void OpenAssetEditor(const TArray<UObject*>& InObjects, TSharedPtr<class IToolkitHost> EditWithinLevelEditor = TSharedPtr<IToolkitHost>()) override;
+	virtual bool ShouldForceWorldCentric() override;
+};
diff --git a/Engine/Plugins/MovieScene/ActorAnimationEditor/Source/ActorAnimationEditor/Private/Factories/ActorAnimationFactoryNew.cpp b/Engine/Plugins/MovieScene/ActorAnimationEditor/Source/ActorAnimationEditor/Private/Factories/ActorAnimationFactoryNew.cpp
new file mode 100644
--- /dev/null
+++ b/Engine/Plugins/MovieScene/ActorAnimationEditor/Source/ActorAnimationEditor/Private/Factories/ActorAnimationFactoryNew.cpp
@@ -0,0 +1,52 @@
+// Copyright 1998-2015 Epic Games, Inc. All Rights Reserved.
+
+#include "ActorAnimationEditorPrivatePCH.h"
+#include "ActorAnimation.h"
+#include "ActorAnimationFactoryNew.h"
+
+
+#define LOCTEXT_NAMESPACE "MovieSceneFactory"
+
+
+/* UActorAnimationFactory structors
+ *****************************************************************************/
+
+UActorAnimationFactoryNew::UActorAnimationFactoryNew(const FObjectInitializer& ObjectInitializer)
+	: Super(ObjectInitializer)
+{
+	// @todo sequencer: Only allow users to create new MovieScenes if that feature is turned on globally.
+	if (FParse::Param(FCommandLine::Get(), TEXT("Sequencer")))
+	{
+		bCreateNew = true;
+	}
+
+	bEditAfterNew = true;
+	SupportedClass = UActorAnimation::StaticClass();
+}
+
+
+/* UFactory interface
+ *****************************************************************************/
+
+UObject* UActorAnimationFactoryNew::FactoryCreateNew(UClass* Class, UObject* InParent, FName Name, EObjectFlags Flags, UObject* Context, FFeedbackContext* Warn)
+{
+	auto NewActorAnimation = NewObject<UActorAnimation>(InParent, Name, Flags);
+	NewActorAnimation->Initialize();
+
+	return NewActorAnimation;
+}
+
+
+bool UActorAnimationFactoryNew::ShouldShowInNewMenu() const
+{
+	// @todo sequencer: Only allow users to create new MovieScenes if that feature is turned on globally.
+	if (!FParse::Param(FCommandLine::Get(), TEXT("Sequencer")))
+	{
+		return false;
+	}
+
+	return true;
+}
+
+
+#undef LOCTEXT_NAMESPACE
\ No newline at end of file
diff --git a/Engine/Plugins/MovieScene/ActorAnimationEditor/Source/ActorAnimationEditor/Private/Factories/ActorAnimationFactoryNew.h b/Engine/Plugins/MovieScene/ActorAnimationEditor/Source/ActorAnimationEditor/Private/Factories/ActorAnimationFactoryNew.h
new file mode 100644
--- /dev/null
+++ b/Engine/Plugins/MovieScene/ActorAnimationEditor/Source/ActorAnimationEditor/Private/Factories/ActorAnimationFactoryNew.h
@@ -0,0 +1,24 @@
+// Copyright 1998-2015 Epic Games, Inc. All Rights Reserved.
+
+#pragma once
+
+#include "Factories/Factory.h"
+#include "ActorAnimationFactoryNew.generated.h"
+
+
+/**
+ * Implements a factory for UActorAnimation objects.
+ */
+UCLASS(hidecategories=Object)
+class UActorAnimationFactoryNew
+	: public UFactory
+{
+	GENERATED_UCLASS_BODY()
+
+public:
+
+	// UFactory Interface
+
+	virtual UObject* FactoryCreateNew(UClass* Class, UObject* InParent, FName Name, EObjectFlags Flags, UObject* Context, FFeedbackContext* Warn) override;
+	virtual bool ShouldShowInNewMenu() const override;
+};
