ActorAnimation: Added support for custom styles

diff --git a/Engine/Plugins/MovieScene/ActorAnimationEditor/Source/ActorAnimationEditor/Private/ActorAnimationEditorModule.cpp b/Engine/Plugins/MovieScene/ActorAnimationEditor/Source/ActorAnimationEditor/Private/ActorAnimationEditorModule.cpp
--- a/Engine/Plugins/MovieScene/ActorAnimationEditor/Source/ActorAnimationEditor/Private/ActorAnimationEditorModule.cpp
+++ b/Engine/Plugins/MovieScene/ActorAnimationEditor/Source/ActorAnimationEditor/Private/ActorAnimationEditorModule.cpp
@@ -1,6 +1,7 @@
 // Copyright 1998-2015 Epic Games, Inc. All Rights Reserved.
 
 #include "ActorAnimationEditorPrivatePCH.h"
+#include "ActorAnimationEditorStyle.h"
 #include "ModuleInterface.h"
 
 
@@ -16,6 +17,8 @@ public:
 
 	virtual void StartupModule() override
 	{
+		Style = MakeShareable(new FActorAnimationEditorStyle());
+
 		RegisterAssetTools();
 	}
 	
@@ -31,7 +34,7 @@ protected:
 	{
 		IAssetTools& AssetTools = FModuleManager::LoadModuleChecked<FAssetToolsModule>("AssetTools").Get();
 
-		RegisterAssetTypeAction(AssetTools, MakeShareable(new FActorAnimationActions));
+		RegisterAssetTypeAction(AssetTools, MakeShareable(new FActorAnimationActions(Style.ToSharedRef())));
 	}
 
 	/**
@@ -40,7 +43,7 @@ protected:
 	 * @param AssetTools The asset tools object to register with.
 	 * @param Action The asset type action to register.
 	 */
-	void RegisterAssetTypeAction( IAssetTools& AssetTools, TSharedRef<IAssetTypeActions> Action )
+	void RegisterAssetTypeAction(IAssetTools& AssetTools, TSharedRef<IAssetTypeActions> Action)
 	{
 		AssetTools.RegisterAssetTypeActions(Action);
 		RegisteredAssetTypeActions.Add(Action);
@@ -66,6 +69,9 @@ private:
 
 	/** The collection of registered asset type actions. */
 	TArray<TSharedRef<IAssetTypeActions>> RegisteredAssetTypeActions;
+
+	/** Holds the plug-ins style set. */
+	TSharedPtr<ISlateStyle> Style;
 };
 
 
diff --git a/Engine/Plugins/MovieScene/ActorAnimationEditor/Source/ActorAnimationEditor/Private/ActorAnimationEditorToolkit.cpp b/Engine/Plugins/MovieScene/ActorAnimationEditor/Source/ActorAnimationEditor/Private/ActorAnimationEditorToolkit.cpp
--- a/Engine/Plugins/MovieScene/ActorAnimationEditor/Source/ActorAnimationEditor/Private/ActorAnimationEditorToolkit.cpp
+++ b/Engine/Plugins/MovieScene/ActorAnimationEditor/Source/ActorAnimationEditor/Private/ActorAnimationEditorToolkit.cpp
@@ -13,12 +13,49 @@
 
 #define LOCTEXT_NAMESPACE "Sequencer"
 
-const FName FActorAnimationEditorToolkit::SequencerMainTabId(TEXT("Sequencer_SequencerMain"));
 
+/* Local constants
+ *****************************************************************************/
+
+const FName FActorAnimationEditorToolkit::SequencerMainTabId(TEXT("Sequencer_SequencerMain"));
 
 namespace SequencerDefs
 {
-	static const FName SequencerAppIdentifier( TEXT( "SequencerApp" ) );
+	static const FName SequencerAppIdentifier(TEXT("SequencerApp"));
+}
+
+
+/* FActorAnimationEditorToolkit structors
+ *****************************************************************************/
+
+FActorAnimationEditorToolkit::FActorAnimationEditorToolkit(const TSharedRef<ISlateStyle>& InStyle)
+	: Style(InStyle)
+{
+	// Register sequencer menu extenders.
+	ISequencerModule& SequencerModule = FModuleManager::Get().LoadModuleChecked<ISequencerModule>( "Sequencer" );
+	int32 NewIndex = SequencerModule.GetMenuExtensibilityManager()->GetExtenderDelegates().Add(
+		FAssetEditorExtender::CreateRaw( this, &FActorAnimationEditorToolkit::GetContextSensitiveSequencerExtender ) );
+	SequencerExtenderHandle = SequencerModule.GetMenuExtensibilityManager()->GetExtenderDelegates()[NewIndex].GetHandle();
+}
+
+
+FActorAnimationEditorToolkit::~FActorAnimationEditorToolkit()
+{
+	Sequencer->Close();
+
+	// Unregister delegates
+	if( FModuleManager::Get().IsModuleLoaded( TEXT( "LevelEditor" ) ) )
+	{
+		auto& LevelEditorModule = FModuleManager::LoadModuleChecked< FLevelEditorModule >( TEXT( "LevelEditor" ) );
+		LevelEditorModule.OnMapChanged().RemoveAll( this );
+	}
+
+	// Un-Register sequencer menu extenders.
+	ISequencerModule& SequencerModule = FModuleManager::Get().LoadModuleChecked<ISequencerModule>( "Sequencer" );
+	SequencerModule.GetMenuExtensibilityManager()->GetExtenderDelegates().RemoveAll( [this]( const FAssetEditorExtender& Extender )
+	{
+		return SequencerExtenderHandle == Extender.GetHandle();
+	} );
 }
 
 
@@ -94,36 +131,6 @@ void FActorAnimationEditorToolkit::Initialize( const EToolkitMode::Type Mode, co
 }
 
 
-FActorAnimationEditorToolkit::FActorAnimationEditorToolkit()
-{
-	// Register sequencer menu extenders.
-	ISequencerModule& SequencerModule = FModuleManager::Get().LoadModuleChecked<ISequencerModule>( "Sequencer" );
-	int32 NewIndex = SequencerModule.GetMenuExtensibilityManager()->GetExtenderDelegates().Add(
-		FAssetEditorExtender::CreateRaw( this, &FActorAnimationEditorToolkit::GetContextSensitiveSequencerExtender ) );
-	SequencerExtenderHandle = SequencerModule.GetMenuExtensibilityManager()->GetExtenderDelegates()[NewIndex].GetHandle();
-}
-
-
-FActorAnimationEditorToolkit::~FActorAnimationEditorToolkit()
-{
-	Sequencer->Close();
-
-	// Unregister delegates
-	if( FModuleManager::Get().IsModuleLoaded( TEXT( "LevelEditor" ) ) )
-	{
-		auto& LevelEditorModule = FModuleManager::LoadModuleChecked< FLevelEditorModule >( TEXT( "LevelEditor" ) );
-		LevelEditorModule.OnMapChanged().RemoveAll( this );
-	}
-
-	// Un-Register sequencer menu extenders.
-	ISequencerModule& SequencerModule = FModuleManager::Get().LoadModuleChecked<ISequencerModule>( "Sequencer" );
-	SequencerModule.GetMenuExtensibilityManager()->GetExtenderDelegates().RemoveAll( [this]( const FAssetEditorExtender& Extender )
-	{
-		return SequencerExtenderHandle == Extender.GetHandle();
-	} );
-}
-
-
 TSharedRef<SDockTab> FActorAnimationEditorToolkit::SpawnTab_SequencerMain(const FSpawnTabArgs& Args)
 {
 	check(Args.GetTabId() == SequencerMainTabId);
diff --git a/Engine/Plugins/MovieScene/ActorAnimationEditor/Source/ActorAnimationEditor/Private/ActorAnimationEditorToolkit.h b/Engine/Plugins/MovieScene/ActorAnimationEditor/Source/ActorAnimationEditor/Private/ActorAnimationEditorToolkit.h
--- a/Engine/Plugins/MovieScene/ActorAnimationEditor/Source/ActorAnimationEditor/Private/ActorAnimationEditorToolkit.h
+++ b/Engine/Plugins/MovieScene/ActorAnimationEditor/Source/ActorAnimationEditor/Private/ActorAnimationEditorToolkit.h
@@ -24,8 +24,12 @@ class FActorAnimationEditorToolkit
 { 
 public:
 
-	/** Default constructor */
-	FActorAnimationEditorToolkit();
+	/**
+	 * Creates and initializes a new instance.
+	 *
+	 * @param InStyle The style set to use.
+	 */
+	FActorAnimationEditorToolkit(const TSharedRef<ISlateStyle>& InStyle);
 
 	/** Virtual destructor */
 	virtual ~FActorAnimationEditorToolkit();
@@ -55,7 +59,6 @@ public:
 	virtual void RegisterTabSpawners(const TSharedRef<FTabManager>& TabManager) override;
 	virtual void UnregisterTabSpawners(const TSharedRef<FTabManager>& TabManager) override;
 
-
 	void UpdateViewports(AActor* ActorToViewThrough) const;
 
 private:
@@ -78,6 +81,9 @@ private:
 
 	FDelegateHandle SequencerExtenderHandle;
 
+	/** Pointer to the style set to use for toolkits. */
+	TSharedRef<ISlateStyle> Style;
+
 private:
 
 	/**	The tab ids for all the tabs used */
diff --git a/Engine/Plugins/MovieScene/ActorAnimationEditor/Source/ActorAnimationEditor/Private/AssetTools/ActorAnimationActions.cpp b/Engine/Plugins/MovieScene/ActorAnimationEditor/Source/ActorAnimationEditor/Private/AssetTools/ActorAnimationActions.cpp
--- a/Engine/Plugins/MovieScene/ActorAnimationEditor/Source/ActorAnimationEditor/Private/AssetTools/ActorAnimationActions.cpp
+++ b/Engine/Plugins/MovieScene/ActorAnimationEditor/Source/ActorAnimationEditor/Private/AssetTools/ActorAnimationActions.cpp
@@ -9,6 +9,14 @@
 #define LOCTEXT_NAMESPACE "AssetTypeActions"
 
 
+/* FActorAnimationActions constructors
+ *****************************************************************************/
+
+FActorAnimationActions::FActorAnimationActions( const TSharedRef<ISlateStyle>& InStyle )
+	: Style(InStyle)
+{ }
+
+
 /* IAssetTypeActions interface
  *****************************************************************************/
 
@@ -55,7 +63,7 @@ void FActorAnimationActions::OpenAssetEditor(const TArray<UObject*>& InObjects,
 		if (ActorAnimation != nullptr)
 		{
 			FSequencerViewParams ViewParams(TEXT("ActorAnimationEditorViewParams"));
-			TSharedRef<FActorAnimationEditorToolkit> Toolkit = MakeShareable(new FActorAnimationEditorToolkit);
+			TSharedRef<FActorAnimationEditorToolkit> Toolkit = MakeShareable(new FActorAnimationEditorToolkit(Style));
 			Toolkit->Initialize(Mode, ViewParams, EditWithinLevelEditor, ActorAnimation, true);
 		}
 	}
diff --git a/Engine/Plugins/MovieScene/ActorAnimationEditor/Source/ActorAnimationEditor/Private/AssetTools/ActorAnimationActions.h b/Engine/Plugins/MovieScene/ActorAnimationEditor/Source/ActorAnimationEditor/Private/AssetTools/ActorAnimationActions.h
--- a/Engine/Plugins/MovieScene/ActorAnimationEditor/Source/ActorAnimationEditor/Private/AssetTools/ActorAnimationActions.h
+++ b/Engine/Plugins/MovieScene/ActorAnimationEditor/Source/ActorAnimationEditor/Private/AssetTools/ActorAnimationActions.h
@@ -11,6 +11,15 @@
 class FActorAnimationActions
 	: public FAssetTypeActions_Base
 {
+public:
+
+	/**
+	 * Creates and initializes a new instance.
+	 *
+	 * @param InStyle The style set to use for asset editor toolkits.
+	 */
+	FActorAnimationActions(const TSharedRef<ISlateStyle>& InStyle);
+
 public:
 	
 	// IAssetTypeActions interface
@@ -21,4 +30,9 @@ public:
 	virtual FColor GetTypeColor() const override;
 	virtual void OpenAssetEditor(const TArray<UObject*>& InObjects, TSharedPtr<class IToolkitHost> EditWithinLevelEditor = TSharedPtr<IToolkitHost>()) override;
 	virtual bool ShouldForceWorldCentric() override;
+
+private:
+
+	/** Pointer to the style set to use for toolkits. */
+	TSharedRef<ISlateStyle> Style;
 };
diff --git a/Engine/Plugins/MovieScene/ActorAnimationEditor/Source/ActorAnimationEditor/Private/Styles/ActorAnimationEditorStyle.h b/Engine/Plugins/MovieScene/ActorAnimationEditor/Source/ActorAnimationEditor/Private/Styles/ActorAnimationEditorStyle.h
new file mode 100644
--- /dev/null
+++ b/Engine/Plugins/MovieScene/ActorAnimationEditor/Source/ActorAnimationEditor/Private/Styles/ActorAnimationEditorStyle.h
@@ -0,0 +1,48 @@
+// Copyright 1998-2015 Epic Games, Inc. All Rights Reserved.
+
+#pragma once
+
+
+#define IMAGE_BRUSH(RelativePath, ...) FSlateImageBrush(RootToContentDir(RelativePath, TEXT(".png")), __VA_ARGS__)
+#define BOX_BRUSH(RelativePath, ...) FSlateBoxBrush(RootToContentDir(RelativePath, TEXT(".png")), __VA_ARGS__)
+#define BORDER_BRUSH(RelativePath, ...) FSlateBorderBrush(RootToContentDir(RelativePath, TEXT(".png")), __VA_ARGS__)
+#define TTF_FONT(RelativePath, ...) FSlateFontInfo(RootToContentDir(RelativePath, TEXT(".ttf")), __VA_ARGS__)
+#define OTF_FONT(RelativePath, ...) FSlateFontInfo(RootToContentDir(RelativePath, TEXT(".otf")), __VA_ARGS__)
+
+
+/**
+ * Implements the visual style of the messaging debugger UI.
+ */
+class FActorAnimationEditorStyle
+	: public FSlateStyleSet
+{
+public:
+
+	/** Default constructor. */
+	 FActorAnimationEditorStyle()
+		 : FSlateStyleSet("MediaPlayerEditorStyle")
+	 {
+		const FVector2D Icon16x16(16.0f, 16.0f);
+		const FVector2D Icon20x20(20.0f, 20.0f);
+		const FVector2D Icon40x40(40.0f, 40.0f);
+
+		SetContentRoot(FPaths::EnginePluginsDir() / TEXT("MovieScene/ActorAnimationEditor/Content"));
+
+		// add styles here
+		
+		FSlateStyleRegistry::RegisterSlateStyle(*this);
+	 }
+
+	 /** Destructor. */
+	 ~FActorAnimationEditorStyle()
+	 {
+		FSlateStyleRegistry::UnRegisterSlateStyle(*this);
+	 }
+};
+
+
+#undef IMAGE_BRUSH
+#undef BOX_BRUSH
+#undef BORDER_BRUSH
+#undef TTF_FONT
+#undef OTF_FONT
