ActorAnimationEditor: Fixed style name

diff --git a/Engine/Plugins/MovieScene/ActorAnimationEditor/Source/ActorAnimationEditor/Private/Styles/ActorAnimationEditorStyle.h b/Engine/Plugins/MovieScene/ActorAnimationEditor/Source/ActorAnimationEditor/Private/Styles/ActorAnimationEditorStyle.h
--- a/Engine/Plugins/MovieScene/ActorAnimationEditor/Source/ActorAnimationEditor/Private/Styles/ActorAnimationEditorStyle.h
+++ b/Engine/Plugins/MovieScene/ActorAnimationEditor/Source/ActorAnimationEditor/Private/Styles/ActorAnimationEditorStyle.h
@@ -20,7 +20,7 @@ public:
 
 	/** Default constructor. */
 	 FActorAnimationEditorStyle()
-		 : FSlateStyleSet("MediaPlayerEditorStyle")
+		 : FSlateStyleSet("ActorAnimationEditorStyle")
 	 {
 		const FVector2D Icon16x16(16.0f, 16.0f);
 		const FVector2D Icon20x20(20.0f, 20.0f);
