ActorAnimationEditor: code cleanup pass

diff --git a/Engine/Plugins/MovieScene/ActorAnimationEditor/Source/ActorAnimationEditor/Private/ActorAnimationEditorToolkit.cpp b/Engine/Plugins/MovieScene/ActorAnimationEditor/Source/ActorAnimationEditor/Private/ActorAnimationEditorToolkit.cpp
--- a/Engine/Plugins/MovieScene/ActorAnimationEditor/Source/ActorAnimationEditor/Private/ActorAnimationEditorToolkit.cpp
+++ b/Engine/Plugins/MovieScene/ActorAnimationEditor/Source/ActorAnimationEditor/Private/ActorAnimationEditorToolkit.cpp
@@ -11,7 +11,7 @@
 #include "SDockTab.h"
 
 
-#define LOCTEXT_NAMESPACE "Sequencer"
+#define LOCTEXT_NAMESPACE "ActorAnimationEditor"
 
 
 /* Local constants
@@ -32,9 +32,9 @@ FActorAnimationEditorToolkit::FActorAnimationEditorToolkit(const TSharedRef<ISla
 	: Style(InStyle)
 {
 	// Register sequencer menu extenders.
-	ISequencerModule& SequencerModule = FModuleManager::Get().LoadModuleChecked<ISequencerModule>( "Sequencer" );
+	ISequencerModule& SequencerModule = FModuleManager::Get().LoadModuleChecked<ISequencerModule>("Sequencer");
 	int32 NewIndex = SequencerModule.GetMenuExtensibilityManager()->GetExtenderDelegates().Add(
-		FAssetEditorExtender::CreateRaw( this, &FActorAnimationEditorToolkit::GetContextSensitiveSequencerExtender ) );
+		FAssetEditorExtender::CreateRaw(this, &FActorAnimationEditorToolkit::HandleMenuExtensibilityGetExtender));
 	SequencerExtenderHandle = SequencerModule.GetMenuExtensibilityManager()->GetExtenderDelegates()[NewIndex].GetHandle();
 }
 
@@ -44,48 +44,23 @@ FActorAnimationEditorToolkit::~FActorAnimationEditorToolkit()
 	Sequencer->Close();
 
 	// Unregister delegates
-	if( FModuleManager::Get().IsModuleLoaded( TEXT( "LevelEditor" ) ) )
+	if (FModuleManager::Get().IsModuleLoaded(TEXT("LevelEditor")))
 	{
-		auto& LevelEditorModule = FModuleManager::LoadModuleChecked< FLevelEditorModule >( TEXT( "LevelEditor" ) );
-		LevelEditorModule.OnMapChanged().RemoveAll( this );
+		auto& LevelEditorModule = FModuleManager::LoadModuleChecked<FLevelEditorModule>(TEXT("LevelEditor"));
+		LevelEditorModule.OnMapChanged().RemoveAll(this);
 	}
 
 	// Un-Register sequencer menu extenders.
-	ISequencerModule& SequencerModule = FModuleManager::Get().LoadModuleChecked<ISequencerModule>( "Sequencer" );
-	SequencerModule.GetMenuExtensibilityManager()->GetExtenderDelegates().RemoveAll( [this]( const FAssetEditorExtender& Extender )
+	ISequencerModule& SequencerModule = FModuleManager::Get().LoadModuleChecked<ISequencerModule>("Sequencer");
+	SequencerModule.GetMenuExtensibilityManager()->GetExtenderDelegates().RemoveAll([this](const FAssetEditorExtender& Extender)
 	{
 		return SequencerExtenderHandle == Extender.GetHandle();
-	} );
+	});
 }
 
 
-void FActorAnimationEditorToolkit::RegisterTabSpawners(const TSharedRef<class FTabManager>& TabManager)
-{
-	if (IsWorldCentricAssetEditor())
-	{
-		return;
-	}
-
-	WorkspaceMenuCategory = TabManager->AddLocalWorkspaceMenuCategory(LOCTEXT("WorkspaceMenu_SequencerAssetEditor", "Sequencer"));
-
-	TabManager->RegisterTabSpawner(SequencerMainTabId, FOnSpawnTab::CreateSP(this, &FActorAnimationEditorToolkit::SpawnTab_SequencerMain))
-		.SetDisplayName(LOCTEXT("SequencerMainTab", "Sequencer"))
-		.SetGroup(WorkspaceMenuCategory.ToSharedRef());
-}
-
-
-void FActorAnimationEditorToolkit::UnregisterTabSpawners(const TSharedRef<class FTabManager>& TabManager)
-{
-	if (!IsWorldCentricAssetEditor())
-	{
-		TabManager->UnregisterTabSpawner( SequencerMainTabId );
-	}
-	
-	// @todo remove when world-centric mode is added
-	FLevelEditorModule& LevelEditorModule = FModuleManager::LoadModuleChecked<FLevelEditorModule>("LevelEditor");
-	LevelEditorModule.AttachSequencer( SNullWidget::NullWidget, nullptr );
-}
-
+/* FActorAnimationEditorToolkit interface
+ *****************************************************************************/
 
 void FActorAnimationEditorToolkit::Initialize( const EToolkitMode::Type Mode, const FSequencerViewParams& InViewParams, const TSharedPtr<IToolkitHost>& InitToolkitHost, UActorAnimation* ActorAnimation, bool bEditWithinLevelEditor )
 {
@@ -93,19 +68,19 @@ void FActorAnimationEditorToolkit::Initialize( const EToolkitMode::Type Mode, co
 		const TSharedRef<FTabManager::FLayout> StandaloneDefaultLayout = FTabManager::NewLayout("Standalone_Sequencer_Layout")
 			->AddArea
 			(
-				FTabManager::NewPrimaryArea()
-					->Split
-					(
-						FTabManager::NewStack()
-						->AddTab(SequencerMainTabId, ETabState::OpenedTab)
-					)
+			FTabManager::NewPrimaryArea()
+			->Split
+			(
+			FTabManager::NewStack()
+			->AddTab(SequencerMainTabId, ETabState::OpenedTab)
+			)
 			);
 
 		const bool bCreateDefaultStandaloneMenu = true;
 		const bool bCreateDefaultToolbar = false;
 		FAssetEditorToolkit::InitAssetEditor(Mode, InitToolkitHost, SequencerDefs::SequencerAppIdentifier, StandaloneDefaultLayout, bCreateDefaultStandaloneMenu, bCreateDefaultToolbar, ActorAnimation);
 	}
-	
+
 	// Init Sequencer
 	FSequencerInitParams SequencerInitParams;
 	{
@@ -123,7 +98,7 @@ void FActorAnimationEditorToolkit::Initialize( const EToolkitMode::Type Mode, co
 		FLevelEditorModule& LevelEditorModule = FModuleManager::LoadModuleChecked<FLevelEditorModule>("LevelEditor");
 
 		LevelEditorModule.AttachSequencer( Sequencer->GetSequencerWidget(), SharedThis(this));
-			
+
 		// We need to find out when the user loads a new map, because we might need to re-create puppet actors
 		// when previewing a MovieScene
 		LevelEditorModule.OnMapChanged().AddRaw(this, &FActorAnimationEditorToolkit::HandleMapChanged);
@@ -131,86 +106,135 @@ void FActorAnimationEditorToolkit::Initialize( const EToolkitMode::Type Mode, co
 }
 
 
-TSharedRef<SDockTab> FActorAnimationEditorToolkit::SpawnTab_SequencerMain(const FSpawnTabArgs& Args)
-{
-	check(Args.GetTabId() == SequencerMainTabId);
+/* IToolkit interface
+ *****************************************************************************/
 
-	return SNew(SDockTab)
-		.Icon( FEditorStyle::GetBrush("Sequencer.Tabs.SequencerMain") )
-		.Label( LOCTEXT("SequencerMainTitle", "Sequencer") )
-		.TabColorScale( GetTabColorScale() )
-		[
-			Sequencer->GetSequencerWidget()
-		];
+FText FActorAnimationEditorToolkit::GetBaseToolkitName() const
+{
+	return LOCTEXT("AppLabel", "Actor Animation Editor");
 }
 
 
 FName FActorAnimationEditorToolkit::GetToolkitFName() const
 {
-	static FName SequencerName("Sequencer");
+	static FName SequencerName("ActorAnimationEditor");
 	return SequencerName;
 }
 
 
-FText FActorAnimationEditorToolkit::GetBaseToolkitName() const
+FLinearColor FActorAnimationEditorToolkit::GetWorldCentricTabColorScale() const
 {
-	return LOCTEXT("AppLabel", "Sequencer");
+	return FLinearColor(0.7, 0.0f, 0.0f, 0.5f);
 }
 
 
-FLinearColor FActorAnimationEditorToolkit::GetWorldCentricTabColorScale() const
+FString FActorAnimationEditorToolkit::GetWorldCentricTabPrefix() const
 {
-	return FLinearColor( 0.7, 0.0f, 0.0f, 0.5f );
+	return LOCTEXT("WorldCentricTabPrefix", "Sequencer ").ToString();
 }
 
 
-FString FActorAnimationEditorToolkit::GetWorldCentricTabPrefix() const
+void FActorAnimationEditorToolkit::RegisterTabSpawners(const TSharedRef<class FTabManager>& TabManager)
 {
-	return LOCTEXT("WorldCentricTabPrefix", "Sequencer ").ToString();
+	if (IsWorldCentricAssetEditor())
+	{
+		return;
+	}
+
+	WorkspaceMenuCategory = TabManager->AddLocalWorkspaceMenuCategory(LOCTEXT("WorkspaceMenu_SequencerAssetEditor", "Sequencer"));
+
+	TabManager->RegisterTabSpawner(SequencerMainTabId, FOnSpawnTab::CreateSP(this, &FActorAnimationEditorToolkit::HandleTabManagerSpawnTab))
+		.SetDisplayName(LOCTEXT("SequencerMainTab", "Sequencer"))
+		.SetGroup(WorkspaceMenuCategory.ToSharedRef())
+		.SetIcon(FSlateIcon(FEditorStyle::GetStyleSetName(), "LevelEditor.Tabs.Viewports"));
 }
 
-TSharedRef<FExtender> FActorAnimationEditorToolkit::GetContextSensitiveSequencerExtender( const TSharedRef<FUICommandList> CommandList, const TArray<UObject*> ContextSensitiveObjects )
+
+void FActorAnimationEditorToolkit::UnregisterTabSpawners(const TSharedRef<class FTabManager>& TabManager)
+{
+	if (!IsWorldCentricAssetEditor())
+	{
+		TabManager->UnregisterTabSpawner( SequencerMainTabId );
+	}
+
+	// @todo remove when world-centric mode is added
+	FLevelEditorModule& LevelEditorModule = FModuleManager::LoadModuleChecked<FLevelEditorModule>("LevelEditor");
+	LevelEditorModule.AttachSequencer( SNullWidget::NullWidget, nullptr );
+}
+
+
+/* FActorAnimationEditorToolkit callbacks
+ *****************************************************************************/
+
+void FActorAnimationEditorToolkit::HandleAddComponentActionExecute( UActorComponent* Component )
+{
+	Sequencer->GetHandleToObject( Component );
+}
+
+
+void FActorAnimationEditorToolkit::HandleMapChanged(class UWorld* NewWorld, EMapChangeType MapChangeType)
+{
+	Sequencer->NotifyMapChanged(NewWorld, MapChangeType);
+}
+
+
+TSharedRef<FExtender> FActorAnimationEditorToolkit::HandleMenuExtensibilityGetExtender( const TSharedRef<FUICommandList> CommandList, const TArray<UObject*> ContextSensitiveObjects )
 {
-	TSharedRef<FExtender> AddTrackMenuExtender( new FExtender() );
+	TSharedRef<FExtender> AddTrackMenuExtender(new FExtender());
 	AddTrackMenuExtender->AddMenuExtension(
 		SequencerMenuExtensionPoints::AddTrackMenu_PropertiesSection,
 		EExtensionHook::Before,
 		CommandList,
-		FMenuExtensionDelegate::CreateRaw( this, &FActorAnimationEditorToolkit::ExtendSequencerAddTrackMenu, ContextSensitiveObjects ) );
+		FMenuExtensionDelegate::CreateRaw(this, &FActorAnimationEditorToolkit::HandleTrackMenuExtensionAddTrack, ContextSensitiveObjects));
+
 	return AddTrackMenuExtender;
 }
 
-void FActorAnimationEditorToolkit::ExtendSequencerAddTrackMenu( FMenuBuilder& AddTrackMenuBuilder, TArray<UObject*> ContextObjects )
+
+TSharedRef<SDockTab> FActorAnimationEditorToolkit::HandleTabManagerSpawnTab(const FSpawnTabArgs& Args)
 {
-	if ( ContextObjects.Num() == 1 )
+	TSharedPtr<SWidget> TabWidget = SNullWidget::NullWidget;
+
+	if (Args.GetTabId() == SequencerMainTabId)
 	{
-		AActor* Actor = Cast<AActor>( ContextObjects[0] );
-		if ( Actor != nullptr )
-		{
-			AddTrackMenuBuilder.BeginSection( "Components", LOCTEXT( "ComponentsSection", "Components" ) );
-			{
-				for ( UActorComponent* Component : Actor->GetComponents() )
-				{
-					FUIAction AddComponentAction( FExecuteAction::CreateSP( this, &FActorAnimationEditorToolkit::AddComponentTrack, Component ) );
-					FText AddComponentLabel = FText::FromString(Component->GetName());
-					FText AddComponentToolTip = FText::Format( LOCTEXT( "ComponentToolTipFormat", "Add {0} component" ), FText::FromString( Component->GetName() ) );
-					AddTrackMenuBuilder.AddMenuEntry( AddComponentLabel, AddComponentToolTip, FSlateIcon(), AddComponentAction );
-				}
-			}
-			AddTrackMenuBuilder.EndSection();
-		}
+		TabWidget = Sequencer->GetSequencerWidget();
 	}
-}
 
-void FActorAnimationEditorToolkit::AddComponentTrack( UActorComponent* Component )
-{
-	Sequencer->GetHandleToObject( Component );
+	return SNew(SDockTab)
+		.Label(LOCTEXT("SequencerMainTitle", "Sequencer"))
+		.TabColorScale(GetTabColorScale())
+		.TabRole(ETabRole::PanelTab)
+		[
+			TabWidget.ToSharedRef()
+		];
 }
 
 
-void FActorAnimationEditorToolkit::HandleMapChanged(class UWorld* NewWorld, EMapChangeType MapChangeType)
+void FActorAnimationEditorToolkit::HandleTrackMenuExtensionAddTrack( FMenuBuilder& AddTrackMenuBuilder, TArray<UObject*> ContextObjects )
 {
-	Sequencer->NotifyMapChanged(NewWorld, MapChangeType);
+	if (ContextObjects.Num() != 1)
+	{
+		return;
+	}
+
+	AActor* Actor = Cast<AActor>(ContextObjects[0]);
+	
+	if (Actor == nullptr)
+	{
+		return;
+	}
+
+	AddTrackMenuBuilder.BeginSection("Components", LOCTEXT("ComponentsSection", "Components"));
+	{
+		for (UActorComponent* Component : Actor->GetComponents())
+		{
+			FUIAction AddComponentAction(FExecuteAction::CreateSP(this, &FActorAnimationEditorToolkit::HandleAddComponentActionExecute, Component));
+			FText AddComponentLabel = FText::FromString(Component->GetName());
+			FText AddComponentToolTip = FText::Format(LOCTEXT("ComponentToolTipFormat", "Add {0} component"), FText::FromString(Component->GetName()));
+			AddTrackMenuBuilder.AddMenuEntry(AddComponentLabel, AddComponentToolTip, FSlateIcon(), AddComponentAction);
+		}
+	}
+	AddTrackMenuBuilder.EndSection();
 }
 
 
diff --git a/Engine/Plugins/MovieScene/ActorAnimationEditor/Source/ActorAnimationEditor/Private/ActorAnimationEditorToolkit.h b/Engine/Plugins/MovieScene/ActorAnimationEditor/Source/ActorAnimationEditor/Private/ActorAnimationEditorToolkit.h
--- a/Engine/Plugins/MovieScene/ActorAnimationEditor/Source/ActorAnimationEditor/Private/ActorAnimationEditorToolkit.h
+++ b/Engine/Plugins/MovieScene/ActorAnimationEditor/Source/ActorAnimationEditor/Private/ActorAnimationEditorToolkit.h
@@ -52,10 +52,10 @@ public:
 
 	// IToolkit interface
 
-	virtual FName GetToolkitFName() const override;
 	virtual FText GetBaseToolkitName() const override;
-	virtual FString GetWorldCentricTabPrefix() const override;
+	virtual FName GetToolkitFName() const override;
 	virtual FLinearColor GetWorldCentricTabColorScale() const override;
+	virtual FString GetWorldCentricTabPrefix() const override;
 	virtual void RegisterTabSpawners(const TSharedRef<FTabManager>& TabManager) override;
 	virtual void UnregisterTabSpawners(const TSharedRef<FTabManager>& TabManager) override;
 
@@ -63,14 +63,21 @@ public:
 
 private:
 
-	EVisibility GetTransportControlVisibility(TWeakPtr<ILevelViewport> LevelViewport) const;
-	TSharedRef<SDockTab> SpawnTab_SequencerMain(const FSpawnTabArgs& Args);	
+	/** Callback for executing the Add Component action. */
+	void HandleAddComponentActionExecute(UActorComponent* Component);
 
-	TSharedRef<FExtender> GetContextSensitiveSequencerExtender(const TSharedRef<FUICommandList> CommandList, const TArray<UObject*> ContextSensitiveObjects);
-	void ExtendSequencerAddTrackMenu(FMenuBuilder& AddTrackMenuBuilder, TArray<UObject*> ContextObjects);
-	void AddComponentTrack(UActorComponent* Component);
+	/** Callback for map changes. */
 	void HandleMapChanged(UWorld* NewWorld, EMapChangeType MapChangeType);
 
+	/** Callback for the menu extensibility manager. */
+	TSharedRef<FExtender> HandleMenuExtensibilityGetExtender(const TSharedRef<FUICommandList> CommandList, const TArray<UObject*> ContextSensitiveObjects);
+
+	/** Callback for spawning tabs. */
+	TSharedRef<SDockTab> HandleTabManagerSpawnTab(const FSpawnTabArgs& Args);	
+
+	/** Callback for the track menu extender. */
+	void HandleTrackMenuExtensionAddTrack(FMenuBuilder& AddTrackMenuBuilder, TArray<UObject*> ContextObjects);
+
 private:
 
 	/** The sequencer used by this editor. */
