Sequencer: Removed -sequencer command line switch. ActorAnimationEditor is disabled by default and must be enabled in Plugin Browser.

diff --git a/Engine/Plugins/MovieScene/ActorAnimationEditor/ActorAnimationEditor.uplugin b/Engine/Plugins/MovieScene/ActorAnimationEditor/ActorAnimationEditor.uplugin
--- a/Engine/Plugins/MovieScene/ActorAnimationEditor/ActorAnimationEditor.uplugin
+++ b/Engine/Plugins/MovieScene/ActorAnimationEditor/ActorAnimationEditor.uplugin
@@ -9,7 +9,7 @@
 	"EngineVersion" : "4.5.0",
 	"Description" : "Content Editor for ActorAnimation Assets.",
 	"Category" : "Editor",
-	"EnabledByDefault" : true,
+	"EnabledByDefault" : false,
 
 	"Modules" :
 	[
diff --git a/Engine/Plugins/MovieScene/ActorAnimationEditor/Source/ActorAnimationEditor/Private/AssetTools/ActorAnimationActions.cpp b/Engine/Plugins/MovieScene/ActorAnimationEditor/Source/ActorAnimationEditor/Private/AssetTools/ActorAnimationActions.cpp
--- a/Engine/Plugins/MovieScene/ActorAnimationEditor/Source/ActorAnimationEditor/Private/AssetTools/ActorAnimationActions.cpp
+++ b/Engine/Plugins/MovieScene/ActorAnimationEditor/Source/ActorAnimationEditor/Private/AssetTools/ActorAnimationActions.cpp
@@ -46,12 +46,6 @@ FColor FActorAnimationActions::GetTypeColor() const
 
 void FActorAnimationActions::OpenAssetEditor(const TArray<UObject*>& InObjects, TSharedPtr<IToolkitHost> EditWithinLevelEditor)
 {
-	// @todo sequencer: Only allow users to create new MovieScenes if that feature is turned on globally.
-	if (!FParse::Param(FCommandLine::Get(), TEXT("Sequencer")))
-	{
-		return;
-	}
-
 	EToolkitMode::Type Mode = EditWithinLevelEditor.IsValid()
 		? EToolkitMode::WorldCentric
 		: EToolkitMode::Standalone;
diff --git a/Engine/Plugins/MovieScene/ActorAnimationEditor/Source/ActorAnimationEditor/Private/Factories/ActorAnimationFactoryNew.cpp b/Engine/Plugins/MovieScene/ActorAnimationEditor/Source/ActorAnimationEditor/Private/Factories/ActorAnimationFactoryNew.cpp
--- a/Engine/Plugins/MovieScene/ActorAnimationEditor/Source/ActorAnimationEditor/Private/Factories/ActorAnimationFactoryNew.cpp
+++ b/Engine/Plugins/MovieScene/ActorAnimationEditor/Source/ActorAnimationEditor/Private/Factories/ActorAnimationFactoryNew.cpp
@@ -14,12 +14,7 @@
 UActorAnimationFactoryNew::UActorAnimationFactoryNew(const FObjectInitializer& ObjectInitializer)
 	: Super(ObjectInitializer)
 {
-	// @todo sequencer: Only allow users to create new MovieScenes if that feature is turned on globally.
-	if (FParse::Param(FCommandLine::Get(), TEXT("Sequencer")))
-	{
-		bCreateNew = true;
-	}
-
+	bCreateNew = true;
 	bEditAfterNew = true;
 	SupportedClass = UActorAnimation::StaticClass();
 }
@@ -39,12 +34,6 @@ UObject* UActorAnimationFactoryNew::FactoryCreateNew(UClass* Class, UObject* InP
 
 bool UActorAnimationFactoryNew::ShouldShowInNewMenu() const
 {
-	// @todo sequencer: Only allow users to create new MovieScenes if that feature is turned on globally.
-	if (!FParse::Param(FCommandLine::Get(), TEXT("Sequencer")))
-	{
-		return false;
-	}
-
 	return true;
 }
 
diff --git a/Engine/Source/Editor/LevelEditor/Private/SLevelEditor.cpp b/Engine/Source/Editor/LevelEditor/Private/SLevelEditor.cpp
--- a/Engine/Source/Editor/LevelEditor/Private/SLevelEditor.cpp
+++ b/Engine/Source/Editor/LevelEditor/Private/SLevelEditor.cpp
@@ -750,9 +750,9 @@ TSharedRef<SDockTab> SLevelEditor::SpawnLevelEditorTab( const FSpawnTabArgs& Arg
 				WorldBrowserModule.CreateWorldBrowserComposition()
 			];
 	}
-	else if( TabIdentifier == TEXT("Sequencer") && FParse::Param(FCommandLine::Get(), TEXT("sequencer")) )
+	else if( TabIdentifier == TEXT("Sequencer") )
 	{
-		// @todo remove when world-centric mode is added
+		// @todo sequencer: remove when world-centric mode is added
 		SequencerTab = SNew(SDockTab)
 			.Icon( FSlateStyleRegistry::FindSlateStyle("ActorAnimationEditorStyle")->GetBrush("ActorAnimationEditor.Tabs.Sequencer") )
 			.Label( NSLOCTEXT("Sequencer", "SequencerMainTitle", "Sequencer") )
