Sequencer: Marked ActorAnimationEditor plug-in as beta version

diff --git a/Engine/Plugins/MovieScene/ActorAnimationEditor/ActorAnimationEditor.uplugin b/Engine/Plugins/MovieScene/ActorAnimationEditor/ActorAnimationEditor.uplugin
--- a/Engine/Plugins/MovieScene/ActorAnimationEditor/ActorAnimationEditor.uplugin
+++ b/Engine/Plugins/MovieScene/ActorAnimationEditor/ActorAnimationEditor.uplugin
@@ -10,6 +10,7 @@
 	"Description" : "Content Editor for ActorAnimation Assets.",
 	"Category" : "Editor",
 	"EnabledByDefault" : false,
+	"IsBetaVersion": true,
 
 	"Modules" :
 	[
