Sequencer: Enabled the "Add" button for tracks

diff --git a/Engine/Plugins/MovieScene/ActorAnimationEditor/Source/ActorAnimationEditor/Private/ActorAnimationEditorToolkit.cpp b/Engine/Plugins/MovieScene/ActorAnimationEditor/Source/ActorAnimationEditor/Private/ActorAnimationEditorToolkit.cpp
--- a/Engine/Plugins/MovieScene/ActorAnimationEditor/Source/ActorAnimationEditor/Private/ActorAnimationEditorToolkit.cpp
+++ b/Engine/Plugins/MovieScene/ActorAnimationEditor/Source/ActorAnimationEditor/Private/ActorAnimationEditorToolkit.cpp
@@ -66,7 +66,7 @@ FActorAnimationEditorToolkit::~FActorAnimationEditorToolkit()
 /* FActorAnimationEditorToolkit interface
  *****************************************************************************/
 
-void FActorAnimationEditorToolkit::Initialize( const EToolkitMode::Type Mode, const FSequencerViewParams& InViewParams, const TSharedPtr<IToolkitHost>& InitToolkitHost, UActorAnimation* ActorAnimation, bool bEditWithinLevelEditor )
+void FActorAnimationEditorToolkit::Initialize( const EToolkitMode::Type Mode, const TSharedPtr<IToolkitHost>& InitToolkitHost, UActorAnimation* ActorAnimation, bool bEditWithinLevelEditor )
 {
 	// create tab layout
 	const TSharedRef<FTabManager::FLayout> StandaloneDefaultLayout = FTabManager::NewLayout("Standalone_ActorAnimationEditor")
@@ -86,9 +86,14 @@ void FActorAnimationEditorToolkit::Initialize( const EToolkitMode::Type Mode, co
 	FAssetEditorToolkit::InitAssetEditor(Mode, InitToolkitHost, SequencerDefs::SequencerAppIdentifier, StandaloneDefaultLayout, bCreateDefaultStandaloneMenu, bCreateDefaultToolbar, ActorAnimation);
 
 	// initialize sequencer
+	FSequencerViewParams ViewParams(TEXT("ActorAnimationEditorViewParams"));
+	{
+		ViewParams.OnGetAddMenuContent = FOnGetAddMenuContent::CreateSP(this, &FActorAnimationEditorToolkit::HandleSequencerGetAddMenuContent);
+	}
+
 	FSequencerInitParams SequencerInitParams;
 	{
-		SequencerInitParams.ViewParams = InViewParams;
+		SequencerInitParams.ViewParams = ViewParams;
 		SequencerInitParams.RootSequence = ActorAnimation;
 		SequencerInitParams.bEditWithinLevelEditor = bEditWithinLevelEditor;
 		SequencerInitParams.ToolkitHost = InitToolkitHost;
@@ -229,6 +234,12 @@ TSharedRef<FExtender> FActorAnimationEditorToolkit::HandleMenuExtensibilityGetEx
 }
 
 
+TSharedRef<SWidget> FActorAnimationEditorToolkit::HandleSequencerGetAddMenuContent(TSharedRef<ISequencer> Sequencer)
+{
+	return SNullWidget::NullWidget;
+}
+
+
 TSharedRef<SDockTab> FActorAnimationEditorToolkit::HandleTabManagerSpawnTab(const FSpawnTabArgs& Args)
 {
 	TSharedPtr<SWidget> TabWidget = SNullWidget::NullWidget;
diff --git a/Engine/Plugins/MovieScene/ActorAnimationEditor/Source/ActorAnimationEditor/Private/ActorAnimationEditorToolkit.h b/Engine/Plugins/MovieScene/ActorAnimationEditor/Source/ActorAnimationEditor/Private/ActorAnimationEditorToolkit.h
--- a/Engine/Plugins/MovieScene/ActorAnimationEditor/Source/ActorAnimationEditor/Private/ActorAnimationEditorToolkit.h
+++ b/Engine/Plugins/MovieScene/ActorAnimationEditor/Source/ActorAnimationEditor/Private/ActorAnimationEditorToolkit.h
@@ -40,13 +40,12 @@ public:
 	 * Initialize this asset editor.
 	 *
 	 * @param Mode Asset editing mode for this editor (standalone or world-centric).
-	 * @param InViewParams Parameters for how to view sequencer UI.
 	 * @param InitToolkitHost When Mode is WorldCentric, this is the level editor instance to spawn this editor within.
 	 * @param ActorAnimation The animation to edit.
 	 * @param TrackEditorDelegates Delegates to call to create auto-key handlers for this sequencer.
 	 * @param bEditWithinLevelEditor Whether or not sequencer should be edited within the level editor.
 	 */
-	void Initialize(const EToolkitMode::Type Mode, const FSequencerViewParams& InViewParams, const TSharedPtr<IToolkitHost>& InitToolkitHost, UActorAnimation* ActorAnimation, bool bEditWithinLevelEditor);
+	void Initialize(const EToolkitMode::Type Mode, const TSharedPtr<IToolkitHost>& InitToolkitHost, UActorAnimation* ActorAnimation, bool bEditWithinLevelEditor);
 
 public:
 
@@ -75,6 +74,9 @@ private:
 	/** Callback for the menu extensibility manager. */
 	TSharedRef<FExtender> HandleMenuExtensibilityGetExtender(const TSharedRef<FUICommandList> CommandList, const TArray<UObject*> ContextSensitiveObjects);
 
+	/** Callback for getting the sequencer's "Add" menu content. */
+	TSharedRef<SWidget> HandleSequencerGetAddMenuContent(TSharedRef<ISequencer> Sequencer);
+
 	/** Callback for spawning tabs. */
 	TSharedRef<SDockTab> HandleTabManagerSpawnTab(const FSpawnTabArgs& Args);	
 
diff --git a/Engine/Plugins/MovieScene/ActorAnimationEditor/Source/ActorAnimationEditor/Private/AssetTools/ActorAnimationActions.cpp b/Engine/Plugins/MovieScene/ActorAnimationEditor/Source/ActorAnimationEditor/Private/AssetTools/ActorAnimationActions.cpp
--- a/Engine/Plugins/MovieScene/ActorAnimationEditor/Source/ActorAnimationEditor/Private/AssetTools/ActorAnimationActions.cpp
+++ b/Engine/Plugins/MovieScene/ActorAnimationEditor/Source/ActorAnimationEditor/Private/AssetTools/ActorAnimationActions.cpp
@@ -56,9 +56,8 @@ void FActorAnimationActions::OpenAssetEditor(const TArray<UObject*>& InObjects,
 
 		if (ActorAnimation != nullptr)
 		{
-			FSequencerViewParams ViewParams(TEXT("ActorAnimationEditorViewParams"));
 			TSharedRef<FActorAnimationEditorToolkit> Toolkit = MakeShareable(new FActorAnimationEditorToolkit(Style));
-			Toolkit->Initialize(Mode, ViewParams, EditWithinLevelEditor, ActorAnimation, true);
+			Toolkit->Initialize(Mode, EditWithinLevelEditor, ActorAnimation, true);
 		}
 	}
 }
