Sequencer: Fix shadowed variable.

diff --git a/Engine/Plugins/MovieScene/ActorAnimationEditor/Source/ActorAnimationEditor/Private/ActorAnimationEditorToolkit.cpp b/Engine/Plugins/MovieScene/ActorAnimationEditor/Source/ActorAnimationEditor/Private/ActorAnimationEditorToolkit.cpp
--- a/Engine/Plugins/MovieScene/ActorAnimationEditor/Source/ActorAnimationEditor/Private/ActorAnimationEditorToolkit.cpp
+++ b/Engine/Plugins/MovieScene/ActorAnimationEditor/Source/ActorAnimationEditor/Private/ActorAnimationEditorToolkit.cpp
@@ -234,7 +234,7 @@ TSharedRef<FExtender> FActorAnimationEditorToolkit::HandleMenuExtensibilityGetEx
 }
 
 
-TSharedRef<SWidget> FActorAnimationEditorToolkit::HandleSequencerGetAddMenuContent(TSharedRef<ISequencer> Sequencer)
+TSharedRef<SWidget> FActorAnimationEditorToolkit::HandleSequencerGetAddMenuContent(TSharedRef<ISequencer> InSequencer)
 {
 	return SNullWidget::NullWidget;
 }
diff --git a/Engine/Plugins/MovieScene/ActorAnimationEditor/Source/ActorAnimationEditor/Private/ActorAnimationEditorToolkit.h b/Engine/Plugins/MovieScene/ActorAnimationEditor/Source/ActorAnimationEditor/Private/ActorAnimationEditorToolkit.h
--- a/Engine/Plugins/MovieScene/ActorAnimationEditor/Source/ActorAnimationEditor/Private/ActorAnimationEditorToolkit.h
+++ b/Engine/Plugins/MovieScene/ActorAnimationEditor/Source/ActorAnimationEditor/Private/ActorAnimationEditorToolkit.h
@@ -75,7 +75,7 @@ private:
 	TSharedRef<FExtender> HandleMenuExtensibilityGetExtender(const TSharedRef<FUICommandList> CommandList, const TArray<UObject*> ContextSensitiveObjects);
 
 	/** Callback for getting the sequencer's "Add" menu content. */
-	TSharedRef<SWidget> HandleSequencerGetAddMenuContent(TSharedRef<ISequencer> Sequencer);
+	TSharedRef<SWidget> HandleSequencerGetAddMenuContent(TSharedRef<ISequencer> InSequencer);
 
 	/** Callback for spawning tabs. */
 	TSharedRef<SDockTab> HandleTabManagerSpawnTab(const FSpawnTabArgs& Args);	
