Sequencer: Added slomo, shot and event track to "Add" button

diff --git a/Engine/Plugins/MovieScene/ActorAnimationEditor/Source/ActorAnimationEditor/Private/ActorAnimationEditorToolkit.cpp b/Engine/Plugins/MovieScene/ActorAnimationEditor/Source/ActorAnimationEditor/Private/ActorAnimationEditorToolkit.cpp
--- a/Engine/Plugins/MovieScene/ActorAnimationEditor/Source/ActorAnimationEditor/Private/ActorAnimationEditorToolkit.cpp
+++ b/Engine/Plugins/MovieScene/ActorAnimationEditor/Source/ActorAnimationEditor/Private/ActorAnimationEditorToolkit.cpp
@@ -215,6 +215,24 @@ void FActorAnimationEditorToolkit::HandleAddComponentMaterialActionExecute( UPri
 }
 
 
+void FActorAnimationEditorToolkit::HandleAddMenuAddEventTrackExecute()
+{
+	
+}
+
+
+void FActorAnimationEditorToolkit::HandleAddMenuAddShotTrackExecute()
+{
+	Sequencer->AddNewShot(FGuid());
+}
+
+
+void FActorAnimationEditorToolkit::HandleAddMenuAddSlomoTrackExecute()
+{
+
+}
+
+
 void FActorAnimationEditorToolkit::HandleMapChanged(class UWorld* NewWorld, EMapChangeType MapChangeType)
 {
 	Sequencer->NotifyMapChanged(NewWorld, MapChangeType);
@@ -236,7 +254,33 @@ TSharedRef<FExtender> FActorAnimationEditorToolkit::HandleMenuExtensibilityGetEx
 
 TSharedRef<SWidget> FActorAnimationEditorToolkit::HandleSequencerGetAddMenuContent(TSharedRef<ISequencer> InSequencer)
 {
-	return SNullWidget::NullWidget;
+	FMenuBuilder MenuBuilder(true, nullptr);
+
+	// add event track
+	MenuBuilder.AddMenuEntry(
+		LOCTEXT("AddEventTrack", "Add Event Track"),
+		LOCTEXT("AddEventTooltip", "Adds a new event track that can trigger events on the timeline."),
+		FSlateIcon(),
+		FUIAction(FExecuteAction::CreateRaw(this, &FActorAnimationEditorToolkit::HandleAddMenuAddEventTrackExecute))
+	);
+
+	// add shot track
+	MenuBuilder.AddMenuEntry(
+		LOCTEXT("AddShotTrack", "Add Shot Track"),
+		LOCTEXT("AddShotTooltip", "Adds a shot track, as well as a new shot at the current scrubber location if a camera is selected."),
+		FSlateIcon(),
+		FUIAction(FExecuteAction::CreateRaw(this, &FActorAnimationEditorToolkit::HandleAddMenuAddShotTrackExecute))
+	);
+
+	// add slomo track
+	MenuBuilder.AddMenuEntry(
+		LOCTEXT("AddSlomoTrack", "Add Slomo Track"),
+		LOCTEXT("AddSlomoTooltip", "Adds a new slomo track that controls the playback speed of the sequence."),
+		FSlateIcon(),
+		FUIAction(FExecuteAction::CreateRaw(this, &FActorAnimationEditorToolkit::HandleAddMenuAddSlomoTrackExecute))
+	);
+
+	return MenuBuilder.MakeWidget();
 }
 
 
diff --git a/Engine/Plugins/MovieScene/ActorAnimationEditor/Source/ActorAnimationEditor/Private/ActorAnimationEditorToolkit.h b/Engine/Plugins/MovieScene/ActorAnimationEditor/Source/ActorAnimationEditor/Private/ActorAnimationEditorToolkit.h
--- a/Engine/Plugins/MovieScene/ActorAnimationEditor/Source/ActorAnimationEditor/Private/ActorAnimationEditorToolkit.h
+++ b/Engine/Plugins/MovieScene/ActorAnimationEditor/Source/ActorAnimationEditor/Private/ActorAnimationEditorToolkit.h
@@ -68,6 +68,15 @@ private:
 	/** Callback for executing the add component material track. */
 	void HandleAddComponentMaterialActionExecute(UPrimitiveComponent* Component, int32 MaterialIndex);
 
+	/** Callback for executing the "Add Event Track" menu action. */
+	void HandleAddMenuAddEventTrackExecute();
+
+	/** Callback for executing the "Add Shot Track" menu action. */
+	void HandleAddMenuAddShotTrackExecute();
+
+	/** Callback for executing the "Add Slomo Track" menu action. */
+	void HandleAddMenuAddSlomoTrackExecute();
+
 	/** Callback for map changes. */
 	void HandleMapChanged(UWorld* NewWorld, EMapChangeType MapChangeType);
 
