Sequencer: Refactored sequencer to not be aware of shot tracks

diff --git a/Engine/Plugins/MovieScene/ActorAnimationEditor/Source/ActorAnimationEditor/Private/ActorAnimationEditorToolkit.cpp b/Engine/Plugins/MovieScene/ActorAnimationEditor/Source/ActorAnimationEditor/Private/ActorAnimationEditorToolkit.cpp
--- a/Engine/Plugins/MovieScene/ActorAnimationEditor/Source/ActorAnimationEditor/Private/ActorAnimationEditorToolkit.cpp
+++ b/Engine/Plugins/MovieScene/ActorAnimationEditor/Source/ActorAnimationEditor/Private/ActorAnimationEditorToolkit.cpp
@@ -223,9 +223,15 @@ void FActorAnimationEditorToolkit::HandleAddMenuAddEventTrackExecute()
 }
 
 
+bool FActorAnimationEditorToolkit::HandleAddMenuAddShotTrackCanExecute()
+{
+	return true;
+}
+
+
 void FActorAnimationEditorToolkit::HandleAddMenuAddShotTrackExecute()
 {
-	Sequencer->AddNewShot(FGuid());
+
 }
 
 
@@ -273,7 +279,10 @@ TSharedRef<SWidget> FActorAnimationEditorToolkit::HandleSequencerGetAddMenuConte
 		LOCTEXT("AddShotTrack", "Add Shot Track"),
 		LOCTEXT("AddShotTooltip", "Adds a shot track, as well as a new shot at the current scrubber location if a camera is selected."),
 		FSlateIcon(Style->GetStyleSetName(), "ActorAnimationEditor.Tracks.Shot"),
-		FUIAction(FExecuteAction::CreateRaw(this, &FActorAnimationEditorToolkit::HandleAddMenuAddShotTrackExecute))
+		FUIAction(
+			FExecuteAction::CreateRaw(this, &FActorAnimationEditorToolkit::HandleAddMenuAddShotTrackExecute),
+			FCanExecuteAction::CreateRaw(this, &FActorAnimationEditorToolkit::HandleAddMenuAddShotTrackCanExecute)
+		)
 	);
 
 	// add slomo track
diff --git a/Engine/Plugins/MovieScene/ActorAnimationEditor/Source/ActorAnimationEditor/Private/ActorAnimationEditorToolkit.h b/Engine/Plugins/MovieScene/ActorAnimationEditor/Source/ActorAnimationEditor/Private/ActorAnimationEditorToolkit.h
--- a/Engine/Plugins/MovieScene/ActorAnimationEditor/Source/ActorAnimationEditor/Private/ActorAnimationEditorToolkit.h
+++ b/Engine/Plugins/MovieScene/ActorAnimationEditor/Source/ActorAnimationEditor/Private/ActorAnimationEditorToolkit.h
@@ -76,6 +76,9 @@ private:
 	/** Callback for executing the "Add Event Track" menu action. */
 	void HandleAddMenuAddEventTrackExecute();
 
+	/** Callback for checking whether the "Add Shot Track" menu action can execute. */
+	bool HandleAddMenuAddShotTrackCanExecute();
+
 	/** Callback for executing the "Add Shot Track" menu action. */
 	void HandleAddMenuAddShotTrackExecute();
 
diff --git a/Engine/Source/Editor/MovieSceneTools/Private/TrackEditors/ShotTrackEditor.cpp b/Engine/Source/Editor/MovieSceneTools/Private/TrackEditors/ShotTrackEditor.cpp
--- a/Engine/Source/Editor/MovieSceneTools/Private/TrackEditors/ShotTrackEditor.cpp
+++ b/Engine/Source/Editor/MovieSceneTools/Private/TrackEditors/ShotTrackEditor.cpp
@@ -580,7 +580,7 @@ void FShotTrackEditor::BuildObjectBindingTrackMenu(FMenuBuilder& MenuBuilder, co
 			NSLOCTEXT("FShotTrackEditor", "AddShot", "Add New Shot"),
 			NSLOCTEXT("FShotTrackEditor", "AddShotTooltip", "Adds a new shot using this camera at the scrubber location."),
 			FSlateIcon(),
-			FUIAction(FExecuteAction::CreateSP(ParentSequencer.Get(), &ISequencer::AddNewShot, ObjectBinding))
+			FUIAction(FExecuteAction::CreateRaw(this, &FShotTrackEditor::HandleAddShotMenuEntryExecute, ObjectBinding))
 			);
 	}
 }
@@ -723,4 +723,10 @@ UFactory* FShotTrackEditor::GetAssetFactoryForNewShot( UClass* SequenceClass )
 	}
 
 	return ShotFactory.Get();
-}
\ No newline at end of file
+}
+
+
+void FShotTrackEditor::HandleAddShotMenuEntryExecute(FGuid CameraGuid)
+{
+	AddKey(CameraGuid);
+}
diff --git a/Engine/Source/Editor/MovieSceneTools/Private/TrackEditors/ShotTrackEditor.h b/Engine/Source/Editor/MovieSceneTools/Private/TrackEditors/ShotTrackEditor.h
--- a/Engine/Source/Editor/MovieSceneTools/Private/TrackEditors/ShotTrackEditor.h
+++ b/Engine/Source/Editor/MovieSceneTools/Private/TrackEditors/ShotTrackEditor.h
@@ -44,6 +44,9 @@ private:
 
 	UFactory* GetAssetFactoryForNewShot( UClass* SequenceClass );
 
+	/** Callback for executing the "Add Shot" menu entry. */
+	void HandleAddShotMenuEntryExecute(FGuid CameraGuid);
+
 private:
 	/** The Thumbnail pool which draws all the viewport thumbnails for the director track */
 	TSharedPtr<class FShotThumbnailPool> ThumbnailPool;
diff --git a/Engine/Source/Editor/Sequencer/Private/Sequencer.cpp b/Engine/Source/Editor/Sequencer/Private/Sequencer.cpp
--- a/Engine/Source/Editor/Sequencer/Private/Sequencer.cpp
+++ b/Engine/Source/Editor/Sequencer/Private/Sequencer.cpp
@@ -375,15 +375,6 @@ void FSequencer::PopToMovieScene( TSharedRef<FMovieSceneSequenceInstance> SubMov
 	}
 }
 
-void FSequencer::AddNewShot(FGuid CameraGuid)
-{
-	// @todo Sequencer: This should not be needed here.  the track editor could be the entry point
-	if (ShotTrackEditor.IsValid())
-	{
-		ShotTrackEditor.Pin()->AddKey(CameraGuid);
-	}
-}
-
 void FSequencer::AddAnimation(FGuid ObjectGuid, class UAnimSequence* AnimSequence)
 {
 	// @todo Sequencer: This should not be needed here.  the track editor could be the entry point
diff --git a/Engine/Source/Editor/Sequencer/Private/Sequencer.h b/Engine/Source/Editor/Sequencer/Private/Sequencer.h
--- a/Engine/Source/Editor/Sequencer/Private/Sequencer.h
+++ b/Engine/Source/Editor/Sequencer/Private/Sequencer.h
@@ -68,7 +68,6 @@ public:
 	virtual TSharedRef<FMovieSceneSequenceInstance> GetFocusedMovieSceneSequenceInstance() const override;
 	virtual void FocusSubMovieScene( TSharedRef<FMovieSceneSequenceInstance> SubMovieSceneInstance ) override;
 	TSharedRef<FMovieSceneSequenceInstance> GetInstanceForSubMovieSceneSection( UMovieSceneSection& SubMovieSceneSection ) const override;
-	virtual void AddNewShot(FGuid CameraGuid) override;
 	virtual void AddAnimation(FGuid ObjectGuid, class UAnimSequence* AnimSequence) override;
 	virtual bool GetAutoKeyEnabled() const override;
 	virtual void SetAutoKeyEnabled(bool bAutoKeyEnabled) override;
diff --git a/Engine/Source/Editor/Sequencer/Public/ISequencer.h b/Engine/Source/Editor/Sequencer/Public/ISequencer.h
--- a/Engine/Source/Editor/Sequencer/Public/ISequencer.h
+++ b/Engine/Source/Editor/Sequencer/Public/ISequencer.h
@@ -111,9 +111,6 @@ public:
 	 */
 	virtual void FilterToSelectedShotSections(bool bZoomToShotBounds = true) = 0;
 	
-	/** Adds a new shot key with the specified camera guid */
-	virtual void AddNewShot(FGuid CameraGuid) = 0;
-
 	/** Adds a new animation to the specified skeletal mesh object guid */
 	virtual void AddAnimation(FGuid ObjectGuid, UAnimSequence* AnimSequence) = 0;
 
