Fix for editor build. Updated another usage of FOnGetAddMenuContent in ActorAniumationEditorToolKit

diff --git a/Engine/Plugins/MovieScene/ActorAnimationEditor/Source/ActorAnimationEditor/Private/ActorAnimationEditorToolkit.cpp b/Engine/Plugins/MovieScene/ActorAnimationEditor/Source/ActorAnimationEditor/Private/ActorAnimationEditorToolkit.cpp
--- a/Engine/Plugins/MovieScene/ActorAnimationEditor/Source/ActorAnimationEditor/Private/ActorAnimationEditorToolkit.cpp
+++ b/Engine/Plugins/MovieScene/ActorAnimationEditor/Source/ActorAnimationEditor/Private/ActorAnimationEditorToolkit.cpp
@@ -262,10 +262,8 @@ TSharedRef<FExtender> FActorAnimationEditorToolkit::HandleMenuExtensibilityGetEx
 }
 
 
-TSharedRef<SWidget> FActorAnimationEditorToolkit::HandleSequencerGetAddMenuContent(TSharedRef<ISequencer> InSequencer)
+void FActorAnimationEditorToolkit::HandleSequencerGetAddMenuContent(FMenuBuilder& MenuBuilder, TSharedRef<ISequencer> InSequencer)
 {
-	FMenuBuilder MenuBuilder(true, nullptr);
-
 	// add event track
 	MenuBuilder.AddMenuEntry(
 		LOCTEXT("AddEventTrack", "Add Event Track"),
@@ -292,8 +290,6 @@ TSharedRef<SWidget> FActorAnimationEditorToolkit::HandleSequencerGetAddMenuConte
 		FSlateIcon(Style->GetStyleSetName(), "ActorAnimationEditor.Tracks.Slomo"),
 		FUIAction(FExecuteAction::CreateRaw(this, &FActorAnimationEditorToolkit::HandleAddMenuAddSlomoTrackExecute))
 	);
-
-	return MenuBuilder.MakeWidget();
 }
 
 
diff --git a/Engine/Plugins/MovieScene/ActorAnimationEditor/Source/ActorAnimationEditor/Private/ActorAnimationEditorToolkit.h b/Engine/Plugins/MovieScene/ActorAnimationEditor/Source/ActorAnimationEditor/Private/ActorAnimationEditorToolkit.h
--- a/Engine/Plugins/MovieScene/ActorAnimationEditor/Source/ActorAnimationEditor/Private/ActorAnimationEditorToolkit.h
+++ b/Engine/Plugins/MovieScene/ActorAnimationEditor/Source/ActorAnimationEditor/Private/ActorAnimationEditorToolkit.h
@@ -92,7 +92,7 @@ private:
 	TSharedRef<FExtender> HandleMenuExtensibilityGetExtender(const TSharedRef<FUICommandList> CommandList, const TArray<UObject*> ContextSensitiveObjects);
 
 	/** Callback for getting the sequencer's "Add" menu content. */
-	TSharedRef<SWidget> HandleSequencerGetAddMenuContent(TSharedRef<ISequencer> InSequencer);
+	void HandleSequencerGetAddMenuContent(FMenuBuilder& MenuBuilder, TSharedRef<ISequencer> InSequencer);
 
 	/** Callback for spawning tabs. */
 	TSharedRef<SDockTab> HandleTabManagerSpawnTab(const FSpawnTabArgs& Args);	
