Sequencer: Forgot to check in

diff --git a/Engine/Plugins/MovieScene/ActorAnimationEditor/Source/ActorAnimationEditor/Private/ActorAnimationEditorToolkit.cpp b/Engine/Plugins/MovieScene/ActorAnimationEditor/Source/ActorAnimationEditor/Private/ActorAnimationEditorToolkit.cpp
--- a/Engine/Plugins/MovieScene/ActorAnimationEditor/Source/ActorAnimationEditor/Private/ActorAnimationEditorToolkit.cpp
+++ b/Engine/Plugins/MovieScene/ActorAnimationEditor/Source/ActorAnimationEditor/Private/ActorAnimationEditorToolkit.cpp
@@ -88,14 +88,8 @@ void FActorAnimationEditorToolkit::Initialize( const EToolkitMode::Type Mode, co
 	FAssetEditorToolkit::InitAssetEditor(Mode, InitToolkitHost, SequencerDefs::SequencerAppIdentifier, StandaloneDefaultLayout, bCreateDefaultStandaloneMenu, bCreateDefaultToolbar, ActorAnimationInstance);
 
 	// initialize sequencer
-	FSequencerViewParams ViewParams(TEXT("ActorAnimationEditorViewParams"));
-	{
-		ViewParams.OnGetAddMenuContent = FOnGetAddMenuContent::CreateSP(this, &FActorAnimationEditorToolkit::HandleSequencerGetAddMenuContent);
-	}
-
 	FSequencerInitParams SequencerInitParams;
 	{
-		SequencerInitParams.ViewParams = ViewParams;
 		SequencerInitParams.RootSequence = ActorAnimationInstance;
 		SequencerInitParams.bEditWithinLevelEditor = bEditWithinLevelEditor;
 		SequencerInitParams.ToolkitHost = InitToolkitHost;
@@ -217,30 +211,6 @@ void FActorAnimationEditorToolkit::HandleAddComponentMaterialActionExecute( UPri
 }
 
 
-void FActorAnimationEditorToolkit::HandleAddMenuAddEventTrackExecute()
-{
-	
-}
-
-
-bool FActorAnimationEditorToolkit::HandleAddMenuAddShotTrackCanExecute()
-{
-	return true;
-}
-
-
-void FActorAnimationEditorToolkit::HandleAddMenuAddShotTrackExecute()
-{
-
-}
-
-
-void FActorAnimationEditorToolkit::HandleAddMenuAddSlomoTrackExecute()
-{
-
-}
-
-
 void FActorAnimationEditorToolkit::HandleMapChanged(class UWorld* NewWorld, EMapChangeType MapChangeType)
 {
 	ActorAnimationInstance->SetContext(NewWorld);
@@ -262,37 +232,6 @@ TSharedRef<FExtender> FActorAnimationEditorToolkit::HandleMenuExtensibilityGetEx
 }
 
 
-void FActorAnimationEditorToolkit::HandleSequencerGetAddMenuContent(FMenuBuilder& MenuBuilder, TSharedRef<ISequencer> InSequencer)
-{
-	// add event track
-	MenuBuilder.AddMenuEntry(
-		LOCTEXT("AddEventTrack", "Add Event Track"),
-		LOCTEXT("AddEventTooltip", "Adds a new event track that can trigger events on the timeline."),
-		FSlateIcon(Style->GetStyleSetName(), "ActorAnimationEditor.Tracks.Event"),
-		FUIAction(FExecuteAction::CreateRaw(this, &FActorAnimationEditorToolkit::HandleAddMenuAddEventTrackExecute))
-	);
-
-	// add shot track
-	MenuBuilder.AddMenuEntry(
-		LOCTEXT("AddShotTrack", "Add Shot Track"),
-		LOCTEXT("AddShotTooltip", "Adds a shot track, as well as a new shot at the current scrubber location if a camera is selected."),
-		FSlateIcon(Style->GetStyleSetName(), "ActorAnimationEditor.Tracks.Shot"),
-		FUIAction(
-			FExecuteAction::CreateRaw(this, &FActorAnimationEditorToolkit::HandleAddMenuAddShotTrackExecute),
-			FCanExecuteAction::CreateRaw(this, &FActorAnimationEditorToolkit::HandleAddMenuAddShotTrackCanExecute)
-		)
-	);
-
-	// add slomo track
-	MenuBuilder.AddMenuEntry(
-		LOCTEXT("AddSlomoTrack", "Add Slomo Track"),
-		LOCTEXT("AddSlomoTooltip", "Adds a new slomo track that controls the playback speed of the sequence."),
-		FSlateIcon(Style->GetStyleSetName(), "ActorAnimationEditor.Tracks.Slomo"),
-		FUIAction(FExecuteAction::CreateRaw(this, &FActorAnimationEditorToolkit::HandleAddMenuAddSlomoTrackExecute))
-	);
-}
-
-
 TSharedRef<SDockTab> FActorAnimationEditorToolkit::HandleTabManagerSpawnTab(const FSpawnTabArgs& Args)
 {
 	TSharedPtr<SWidget> TabWidget = SNullWidget::NullWidget;
diff --git a/Engine/Plugins/MovieScene/ActorAnimationEditor/Source/ActorAnimationEditor/Private/ActorAnimationEditorToolkit.h b/Engine/Plugins/MovieScene/ActorAnimationEditor/Source/ActorAnimationEditor/Private/ActorAnimationEditorToolkit.h
--- a/Engine/Plugins/MovieScene/ActorAnimationEditor/Source/ActorAnimationEditor/Private/ActorAnimationEditorToolkit.h
+++ b/Engine/Plugins/MovieScene/ActorAnimationEditor/Source/ActorAnimationEditor/Private/ActorAnimationEditorToolkit.h
@@ -73,27 +73,12 @@ private:
 	/** Callback for executing the add component material track. */
 	void HandleAddComponentMaterialActionExecute(UPrimitiveComponent* Component, int32 MaterialIndex);
 
-	/** Callback for executing the "Add Event Track" menu action. */
-	void HandleAddMenuAddEventTrackExecute();
-
-	/** Callback for checking whether the "Add Shot Track" menu action can execute. */
-	bool HandleAddMenuAddShotTrackCanExecute();
-
-	/** Callback for executing the "Add Shot Track" menu action. */
-	void HandleAddMenuAddShotTrackExecute();
-
-	/** Callback for executing the "Add Slomo Track" menu action. */
-	void HandleAddMenuAddSlomoTrackExecute();
-
 	/** Callback for map changes. */
 	void HandleMapChanged(UWorld* NewWorld, EMapChangeType MapChangeType);
 
 	/** Callback for the menu extensibility manager. */
 	TSharedRef<FExtender> HandleMenuExtensibilityGetExtender(const TSharedRef<FUICommandList> CommandList, const TArray<UObject*> ContextSensitiveObjects);
 
-	/** Callback for getting the sequencer's "Add" menu content. */
-	void HandleSequencerGetAddMenuContent(FMenuBuilder& MenuBuilder, TSharedRef<ISequencer> InSequencer);
-
 	/** Callback for spawning tabs. */
 	TSharedRef<SDockTab> HandleTabManagerSpawnTab(const FSpawnTabArgs& Args);	
 
