Sequencer: Add actor animation icon in cinematics level editor menu.

diff --git a/Engine/Plugins/MovieScene/ActorAnimationEditor/Source/ActorAnimationEditor/Private/ActorAnimationEditorModule.cpp b/Engine/Plugins/MovieScene/ActorAnimationEditor/Source/ActorAnimationEditor/Private/ActorAnimationEditorModule.cpp
--- a/Engine/Plugins/MovieScene/ActorAnimationEditor/Source/ActorAnimationEditor/Private/ActorAnimationEditorModule.cpp
+++ b/Engine/Plugins/MovieScene/ActorAnimationEditor/Source/ActorAnimationEditor/Private/ActorAnimationEditorModule.cpp
@@ -104,7 +104,7 @@ protected:
 		// Create and register the level editor toolbar menu extension
 		CinematicsMenuExtender = MakeShareable(new FExtender);
 		CinematicsMenuExtender->AddMenuExtension("LevelEditorNewMatinee", EExtensionHook::First, CommandList, FMenuExtensionDelegate::CreateStatic([](FMenuBuilder& MenuBuilder){
-			MenuBuilder.AddMenuEntry(FActorAnimationExtensionCommands::Get().CreateNewActorAnimationInLevel);
+			MenuBuilder.AddMenuEntry(FActorAnimationExtensionCommands::Get().CreateNewActorAnimationInLevel, NAME_None, TAttribute<FText>(), TAttribute<FText>(), FSlateIcon(FEditorStyle::GetStyleSetName(), "ClassIcon.ActorAnimation") );
 		}));
 
 		FLevelEditorModule& LevelEditorModule = FModuleManager::LoadModuleChecked<FLevelEditorModule>("LevelEditor");
