Back out revision 4 from //depot/UE4/Engine/Plugins/MovieScene/ActorAnimationEditor/Source/ActorAnimationEditor/Private/ActorAnimationEditorModule.cpp

diff --git a/Engine/Plugins/MovieScene/ActorAnimationEditor/Source/ActorAnimationEditor/Private/ActorAnimationEditorModule.cpp b/Engine/Plugins/MovieScene/ActorAnimationEditor/Source/ActorAnimationEditor/Private/ActorAnimationEditorModule.cpp
--- a/Engine/Plugins/MovieScene/ActorAnimationEditor/Source/ActorAnimationEditor/Private/ActorAnimationEditorModule.cpp
+++ b/Engine/Plugins/MovieScene/ActorAnimationEditor/Source/ActorAnimationEditor/Private/ActorAnimationEditorModule.cpp
@@ -104,7 +104,7 @@ protected:
 		// Create and register the level editor toolbar menu extension
 		CinematicsMenuExtender = MakeShareable(new FExtender);
 		CinematicsMenuExtender->AddMenuExtension("LevelEditorNewMatinee", EExtensionHook::First, CommandList, FMenuExtensionDelegate::CreateStatic([](FMenuBuilder& MenuBuilder){
-			MenuBuilder.AddMenuEntry(FActorAnimationExtensionCommands::Get().CreateNewActorAnimationInLevel, NAME_None, TAttribute<FText>(), TAttribute<FText>(), FSlateIcon(FEditorStyle::GetStyleSetName(), "ClassIcon.ActorAnimation") );
+			MenuBuilder.AddMenuEntry(FActorAnimationExtensionCommands::Get().CreateNewActorAnimationInLevel);
 		}));
 
 		FLevelEditorModule& LevelEditorModule = FModuleManager::LoadModuleChecked<FLevelEditorModule>("LevelEditor");
