Extended add menu on sequencer to allow adding actors

diff --git a/Engine/Plugins/MovieScene/ActorAnimationEditor/Source/ActorAnimationEditor/ActorAnimationEditor.Build.cs b/Engine/Plugins/MovieScene/ActorAnimationEditor/Source/ActorAnimationEditor/ActorAnimationEditor.Build.cs
--- a/Engine/Plugins/MovieScene/ActorAnimationEditor/Source/ActorAnimationEditor/ActorAnimationEditor.Build.cs
+++ b/Engine/Plugins/MovieScene/ActorAnimationEditor/Source/ActorAnimationEditor/ActorAnimationEditor.Build.cs
@@ -9,6 +9,7 @@ public class ActorAnimationEditor : ModuleRules
         DynamicallyLoadedModuleNames.AddRange(
             new string[] {
 				"AssetTools",
+				"SceneOutliner",
 			}
         );
         
@@ -33,6 +34,7 @@ public class ActorAnimationEditor : ModuleRules
 		PrivateIncludePathModuleNames.AddRange(
 			new string[] {
 				"AssetTools",
+				"SceneOutliner",
 			}
 		);
 
diff --git a/Engine/Plugins/MovieScene/ActorAnimationEditor/Source/ActorAnimationEditor/Private/ActorAnimationEditorToolkit.cpp b/Engine/Plugins/MovieScene/ActorAnimationEditor/Source/ActorAnimationEditor/Private/ActorAnimationEditorToolkit.cpp
--- a/Engine/Plugins/MovieScene/ActorAnimationEditor/Source/ActorAnimationEditor/Private/ActorAnimationEditorToolkit.cpp
+++ b/Engine/Plugins/MovieScene/ActorAnimationEditor/Source/ActorAnimationEditor/Private/ActorAnimationEditorToolkit.cpp
@@ -13,6 +13,10 @@
 #include "MovieScene.h"
 #include "MovieSceneMaterialTrack.h"
 #include "ScopedTransaction.h"
+#include "ClassIconFinder.h"
+
+#include "SceneOutlinerModule.h"
+#include "SceneOutlinerPublicTypes.h"
 
 
 #define LOCTEXT_NAMESPACE "ActorAnimationEditor"
@@ -29,9 +33,6 @@ namespace SequencerDefs
 }
 
 
-/* FActorAnimationEditorToolkit structors
- *****************************************************************************/
-
 FActorAnimationEditorToolkit::FActorAnimationEditorToolkit(const TSharedRef<ISlateStyle>& InStyle)
 	: Style(InStyle)
 {
@@ -93,6 +94,24 @@ void FActorAnimationEditorToolkit::Initialize( const EToolkitMode::Type Mode, co
 		SequencerInitParams.RootSequence = ActorAnimationInstance;
 		SequencerInitParams.bEditWithinLevelEditor = bEditWithinLevelEditor;
 		SequencerInitParams.ToolkitHost = InitToolkitHost;
+
+		TSharedRef<FExtender> AddMenuExtender = MakeShareable(new FExtender);
+
+		AddMenuExtender->AddMenuExtension("AddTracks", EExtensionHook::Before, nullptr,
+			FMenuExtensionDelegate::CreateLambda([=](FMenuBuilder& MenuBuilder){
+
+				MenuBuilder.AddSubMenu(
+					LOCTEXT("AddActor_Label", "Add Actor To Sequencer"),
+					LOCTEXT("AddActor_ToolTip", "Allow sequencer to possess an actor that already exists in the current level"),
+					FNewMenuDelegate::CreateRaw(this, &FActorAnimationEditorToolkit::AddPosessActorMenuExtensions),
+					false /*bInOpenSubMenuOnClick*/,
+					FSlateIcon("ActorAnimationEditorStyle", "ActorAnimationEditor.PossessNewActor")
+					);
+
+			})
+		);
+
+		SequencerInitParams.ViewParams.AddMenuExtender = AddMenuExtender;
 	}
 
 	Sequencer = FModuleManager::LoadModuleChecked<ISequencerModule>("Sequencer").CreateSequencer(SequencerInitParams);
@@ -210,6 +229,112 @@ void FActorAnimationEditorToolkit::HandleAddComponentMaterialActionExecute( UPri
 	}
 }
 
+void FActorAnimationEditorToolkit::AddActorsToSequencer(AActor*const* InActors, int32 NumActors)
+{
+	const FScopedTransaction Transaction(LOCTEXT("UndoPossessingObject", "Possess Object(s) with Sequencer"));
+
+	UMovieSceneSequence* FocussedSequence = Sequencer->GetFocusedMovieSceneSequence();
+	UMovieScene* FocussedMovieScene = FocussedSequence->GetMovieScene();
+
+	FocussedSequence->Modify();
+
+	GEditor->SelectNone(true, true);
+	while (NumActors--)
+	{
+		AActor* ThisActor = *InActors;
+
+		FGuid ObjectGuid = FocussedSequence->FindObjectId(*ThisActor);
+
+		// Add this object if it hasn't already been possessed
+		if (!ObjectGuid.IsValid() || !FocussedMovieScene->FindPossessable(ObjectGuid))
+		{
+			ObjectGuid = FocussedMovieScene->AddPossessable(ThisActor->GetName(), ThisActor->GetClass());
+			FocussedSequence->BindPossessableObject(ObjectGuid, *ThisActor);
+		}
+
+		GEditor->SelectActor(ThisActor, true, true);
+
+		InActors++;
+	}
+
+	Sequencer->NotifyMovieSceneDataChanged();
+}
+
+void FActorAnimationEditorToolkit::AddPosessActorMenuExtensions(FMenuBuilder& MenuBuilder)
+{
+	auto IsActorValidForPossession = [=](const AActor* Actor){
+		bool bCreateHandleIfMissing = false;
+		return !Sequencer->GetHandleToObject((UObject*)Actor, bCreateHandleIfMissing).IsValid();
+	};
+
+	// Set up a menu entry to add the selected actor(s) to the sequencer
+	TArray<AActor*> ActorsValidForPossession;
+	GEditor->GetSelectedActors()->GetSelectedObjects(ActorsValidForPossession);
+	ActorsValidForPossession.RemoveAll([&](AActor* In){ return !IsActorValidForPossession(In); });
+
+	FText SelectedLabel;
+	FSlateIcon ActorIcon(FEditorStyle::GetStyleSetName(), FClassIconFinder::FindIconNameForClass(AActor::StaticClass()));
+	if (ActorsValidForPossession.Num() == 1)
+	{
+		SelectedLabel = FText::Format(LOCTEXT("AddSpecificActor", "Add '{0}'"), FText::FromString(ActorsValidForPossession[0]->GetActorLabel()));
+		FName IconName = FClassIconFinder::FindIconNameForActor(ActorsValidForPossession[0]);
+		ActorIcon = FSlateIcon(FEditorStyle::GetStyleSetName(), FClassIconFinder::FindIconNameForClass(ActorsValidForPossession[0]->GetClass()));
+	}
+	else if (ActorsValidForPossession.Num() > 1)
+	{
+		SelectedLabel = FText::Format(LOCTEXT("AddSpecificActor", "Add Current Selection ({0} actors)"), FText::AsNumber(ActorsValidForPossession.Num()));
+	}
+
+	if (!SelectedLabel.IsEmpty())
+	{
+		// Copy the array into the lambda - probably not that big a deal
+		MenuBuilder.AddMenuEntry(SelectedLabel, FText(), ActorIcon, FExecuteAction::CreateLambda([=]{
+			FSlateApplication::Get().DismissAllMenus();
+			AddActorsToSequencer(ActorsValidForPossession.GetData(), ActorsValidForPossession.Num());
+		}));
+	}
+	
+	MenuBuilder.BeginSection("ChooseActorSection", LOCTEXT("ChooseActor", "Choose Actor:"));
+
+	using namespace SceneOutliner;
+
+	// Set up a menu entry to add any arbitrary actor to the sequencer
+	FInitializationOptions InitOptions;
+	{
+		InitOptions.Mode = ESceneOutlinerMode::ActorPicker;
+
+		// We hide the header row to keep the UI compact.
+		InitOptions.bShowHeaderRow = false;
+		InitOptions.bShowSearchBox = true;
+		InitOptions.bShowCreateNewFolder = false;
+		// Only want the actor label column
+		InitOptions.ColumnMap.Add(FBuiltInColumnTypes::Label(), FColumnInfo(EColumnVisibility::Visible, 0));
+
+		// Only display actors that are not possessed already
+		InitOptions.Filters->AddFilterPredicate( FActorFilterPredicate::CreateLambda( IsActorValidForPossession ) );
+	}
+
+	// actor selector to allow the user to choose an actor
+	FSceneOutlinerModule& SceneOutlinerModule = FModuleManager::LoadModuleChecked<FSceneOutlinerModule>("SceneOutliner");
+	TSharedRef< SWidget > MiniSceneOutliner =
+		SNew( SVerticalBox )
+		+SVerticalBox::Slot()
+		.AutoHeight()
+		.MaxHeight(400.0f)
+		[
+			SceneOutlinerModule.CreateSceneOutliner(
+				InitOptions,
+				FOnActorPicked::CreateLambda([=](AActor* Actor){
+					// Create a new binding for this actor
+					FSlateApplication::Get().DismissAllMenus();
+					AddActorsToSequencer(&Actor, 1);
+				})
+			)
+		];
+	MenuBuilder.AddWidget(MiniSceneOutliner, FText::GetEmpty(), true);
+	MenuBuilder.EndSection();
+}
+
 
 void FActorAnimationEditorToolkit::HandleMapChanged(class UWorld* NewWorld, EMapChangeType MapChangeType)
 {
diff --git a/Engine/Plugins/MovieScene/ActorAnimationEditor/Source/ActorAnimationEditor/Private/ActorAnimationEditorToolkit.h b/Engine/Plugins/MovieScene/ActorAnimationEditor/Source/ActorAnimationEditor/Private/ActorAnimationEditorToolkit.h
--- a/Engine/Plugins/MovieScene/ActorAnimationEditor/Source/ActorAnimationEditor/Private/ActorAnimationEditorToolkit.h
+++ b/Engine/Plugins/MovieScene/ActorAnimationEditor/Source/ActorAnimationEditor/Private/ActorAnimationEditorToolkit.h
@@ -85,6 +85,12 @@ private:
 	/** Callback for the track menu extender. */
 	void HandleTrackMenuExtensionAddTrack(FMenuBuilder& AddTrackMenuBuilder, TArray<UObject*> ContextObjects);
 
+	/** Menu extension callback for the add menu */
+	void AddPosessActorMenuExtensions(FMenuBuilder& MenuBuilder);
+
+	/** Add the specified actors to the sequencer */
+	void AddActorsToSequencer(AActor*const* InActors, int32 NumActors);
+
 private:
 
 	/** The sequencer used by this editor. */
diff --git a/Engine/Plugins/MovieScene/ActorAnimationEditor/Source/ActorAnimationEditor/Private/Styles/ActorAnimationEditorStyle.h b/Engine/Plugins/MovieScene/ActorAnimationEditor/Source/ActorAnimationEditor/Private/Styles/ActorAnimationEditorStyle.h
--- a/Engine/Plugins/MovieScene/ActorAnimationEditor/Source/ActorAnimationEditor/Private/Styles/ActorAnimationEditorStyle.h
+++ b/Engine/Plugins/MovieScene/ActorAnimationEditor/Source/ActorAnimationEditor/Private/Styles/ActorAnimationEditorStyle.h
@@ -30,6 +30,8 @@ public:
 
 		// tab icons
 		Set("ActorAnimationEditor.Tabs.Sequencer", new IMAGE_BRUSH("icon_tab_sequencer_16x", Icon16x16));
+		Set("ActorAnimationEditor.PossessNewActor", new IMAGE_BRUSH("PossessNewActor_16x", Icon16x16));
+		Set("ActorAnimationEditor.PossessNewActor.Small", new IMAGE_BRUSH("PossessNewActor_16x", Icon16x16));
 		Set("ActorAnimationEditor.CreateNewActorAnimationInLevel", new IMAGE_BRUSH("CreateNewActorAnimationInLevel_16x", Icon16x16));
 		Set("ActorAnimationEditor.CreateNewActorAnimationInLevel.Small", new IMAGE_BRUSH("CreateNewActorAnimationInLevel_16x", Icon16x16));
 
diff --git a/Engine/Source/Editor/EditorStyle/Private/SlateEditorStyle.cpp b/Engine/Source/Editor/EditorStyle/Private/SlateEditorStyle.cpp
--- a/Engine/Source/Editor/EditorStyle/Private/SlateEditorStyle.cpp
+++ b/Engine/Source/Editor/EditorStyle/Private/SlateEditorStyle.cpp
@@ -1601,8 +1601,6 @@ void FSlateEditorStyle::FStyle::SetupGeneralStyles()
 		Set( "Sequencer.AnimationOutliner.RegularFont", TTF_CORE_FONT( "Fonts/Roboto-Regular", 9 ) );
 		Set( "Sequencer.ShotFilter", new IMAGE_BRUSH( "Sequencer/FilteredArea", FVector2D(74,74), FLinearColor::White, ESlateBrushTileType::Both ) );
 		Set( "Sequencer.KeyMark", new IMAGE_BRUSH("Sequencer/KeyMark", FVector2D(3,21), FLinearColor::White, ESlateBrushTileType::NoTile ) );
-		Set( "Sequencer.AddObject", new IMAGE_BRUSH( "Icons/SequencerIcons/icon_Sequencer_NewTrack_24x", Icon40x40 ) );
-		Set( "Sequencer.AddObject.Small", new IMAGE_BRUSH( "Icons/SequencerIcons/icon_Sequencer_NewTrack_24x", Icon20x20 ) );
 		Set( "Sequencer.ToggleAutoKeyEnabled", new IMAGE_BRUSH( "Icons/SequencerIcons/icon_Sequencer_AutoKey_24x", Icon40x40 ) );
 		Set( "Sequencer.ToggleAutoKeyEnabled.Small", new IMAGE_BRUSH( "Icons/SequencerIcons/icon_Sequencer_AutoKey_24x", Icon20x20 ) );
 		Set( "Sequencer.KeyAllEnabled", new IMAGE_BRUSH( "Icons/SequencerIcons/icon_Sequencer_WholeKeyframing_24x", Icon40x40 ) );
diff --git a/Engine/Source/Editor/Sequencer/Private/SSequencer.cpp b/Engine/Source/Editor/Sequencer/Private/SSequencer.cpp
--- a/Engine/Source/Editor/Sequencer/Private/SSequencer.cpp
+++ b/Engine/Source/Editor/Sequencer/Private/SSequencer.cpp
@@ -205,6 +205,7 @@ void SSequencer::Construct( const FArguments& InArgs, TSharedRef< class FSequenc
 	TSharedRef<ITimeSlider> BottomTimeRange = SequencerWidgets.CreateTimeRange( TimeSliderController, TAttribute<EVisibility>(this, &SSequencer::GetTimeRangeVisibility), TAttribute<bool>(this, &SSequencer::ShowFrameNumbers), TAttribute<float>(this, &SSequencer::OnGetTimeSnapInterval));
 
 	OnGetAddMenuContent = InArgs._OnGetAddMenuContent;
+	AddMenuExtender = InArgs._AddMenuExtender;
 
 	ColumnFillCoefficients[0] = .25f;
 	ColumnFillCoefficients[1] = .75f;
@@ -642,13 +643,6 @@ TSharedRef<SWidget> SSequencer::MakeToolBar()
 
 		if( Sequencer.Pin()->IsLevelEditorSequencer() )
 		{
-			ToolBarBuilder.AddToolBarButton(
-				FUIAction(FExecuteAction::CreateSP(this, &SSequencer::OnAddObjectClicked))
-				, NAME_None
-				, LOCTEXT("AddObject", "Add Object")
-				, LOCTEXT("AddObjectTooltip", "Add selected objects")
-				, FSlateIcon(FEditorStyle::GetStyleSetName(), "Sequencer.AddObject"));
-
 			TAttribute<FSlateIcon> KeyAllIcon;
 			KeyAllIcon.Bind(TAttribute<FSlateIcon>::FGetter::CreateLambda([&]{
 				static FSlateIcon KeyAllEnabledIcon(FEditorStyle::GetStyleSetName(), "Sequencer.KeyAllEnabled");
@@ -766,18 +760,24 @@ TSharedRef<SWidget> SSequencer::MakeToolBar()
 
 TSharedRef<SWidget> SSequencer::MakeAddMenu()
 {
-	FMenuBuilder MenuBuilder(true, nullptr);
+	FMenuBuilder MenuBuilder(true, nullptr, AddMenuExtender);
 	{
+
 		// let toolkits populate the menu
+		MenuBuilder.BeginSection("MainMenu");
 		OnGetAddMenuContent.ExecuteIfBound(MenuBuilder, Sequencer.Pin().ToSharedRef());
+		MenuBuilder.EndSection();
 
 		// let track editors populate the menu
 		TSharedPtr<FSequencer> PinnedSequencer = Sequencer.Pin();
 
+		// Always create the section so that we afford extension
+		MenuBuilder.BeginSection("AddTracks");
 		if (PinnedSequencer.IsValid())
 		{
 			PinnedSequencer->BuildAddTrackMenu(MenuBuilder);
 		}
+		MenuBuilder.EndSection();
 	}
 
 	return MenuBuilder.MakeWidget();
@@ -1291,6 +1291,16 @@ void SSequencer::OnActorSelectionChanged(UObject*)
 			}
 		}
 	}
+
+	const TSet<TSharedRef<FSequencerDisplayNode>>& OutlinerSelection = Sequencer.Pin()->GetSelection().GetSelectedOutlinerNodes();
+	if (OutlinerSelection.Num() == 1)
+	{
+		for (auto& Node : OutlinerSelection)
+		{
+			TreeView->RequestScrollIntoView(Node);
+			break;
+		}
+	}
 }
 
 void SSequencer::OnCrumbClicked(const FSequencerBreadcrumb& Item)
@@ -1355,11 +1365,6 @@ void SSequencer::OnSaveMovieSceneClicked()
 	Sequencer.Pin()->SaveCurrentMovieScene();
 }
 
-void SSequencer::OnAddObjectClicked()
-{
-	Sequencer.Pin()->AddSelectedObjects();
-}
-
 void SSequencer::StepToNextKey()
 {
 	StepToKey(true, false);
diff --git a/Engine/Source/Editor/Sequencer/Private/SSequencer.h b/Engine/Source/Editor/Sequencer/Private/SSequencer.h
--- a/Engine/Source/Editor/Sequencer/Private/SSequencer.h
+++ b/Engine/Source/Editor/Sequencer/Private/SSequencer.h
@@ -83,6 +83,8 @@ public:
 		SLATE_EVENT( FOnScrubPositionChanged, OnScrubPositionChanged )
 		/** Called to populate the add combo button in the toolbar. */
 		SLATE_EVENT( FOnGetAddMenuContent, OnGetAddMenuContent )
+		/** Extender to use for the add menu. */
+		SLATE_ARGUMENT( TSharedPtr<FExtender>, AddMenuExtender )
 	SLATE_END_ARGS()
 
 
@@ -127,9 +129,6 @@ public:
 	/** Called when the save button is clicked */
 	void OnSaveMovieSceneClicked();
 
-	/** Called when the add object button is clicked */
-	void OnAddObjectClicked();
-
 	/** Access the tree view for this sequencer */
 	TSharedPtr<SSequencerTreeView> GetTreeView() const;
 
@@ -330,6 +329,8 @@ private:
 	bool bUserIsSelecting;
 	/** The current edit tool */
 	TUniquePtr<ISequencerEditTool> EditTool;
+	/** Extender to use for the 'add' menu */
+	TSharedPtr<FExtender> AddMenuExtender;
 
 	FOnGetAddMenuContent OnGetAddMenuContent;
 
diff --git a/Engine/Source/Editor/Sequencer/Private/Sequencer.cpp b/Engine/Source/Editor/Sequencer/Private/Sequencer.cpp
--- a/Engine/Source/Editor/Sequencer/Private/Sequencer.cpp
+++ b/Engine/Source/Editor/Sequencer/Private/Sequencer.cpp
@@ -114,7 +114,8 @@ void FSequencer::InitSequencer(const FSequencerInitParams& InitParams, const TSh
 			.OnScrubPositionChanged( this, &FSequencer::OnScrubPositionChanged )
 			.OnViewRangeChanged( this, &FSequencer::OnViewRangeChanged )
 			.OnClampRangeChanged( this, &FSequencer::OnClampRangeChanged )
-			.OnGetAddMenuContent(InitParams.ViewParams.OnGetAddMenuContent);
+			.OnGetAddMenuContent(InitParams.ViewParams.OnGetAddMenuContent)
+			.AddMenuExtender(InitParams.ViewParams.AddMenuExtender);
 
 		// When undo occurs, get a notification so we can make sure our view is up to date
 		GEditor->RegisterForUndo(this);
diff --git a/Engine/Source/Editor/Sequencer/Public/ISequencerModule.h b/Engine/Source/Editor/Sequencer/Public/ISequencerModule.h
--- a/Engine/Source/Editor/Sequencer/Public/ISequencerModule.h
+++ b/Engine/Source/Editor/Sequencer/Public/ISequencerModule.h
@@ -40,6 +40,9 @@ struct FSequencerViewParams
 
 	FOnGetAddMenuContent OnGetAddMenuContent;
 
+	/** A menu extender for the add menu */
+	TSharedPtr<FExtender> AddMenuExtender;
+
 	/** Unique name for the sequencer. */
 	FString UniqueName;
 
