Sequencer: Fix level editor sequencer settings so that there's only one set of sequencer settings.

diff --git a/Engine/Plugins/MovieScene/ActorAnimationEditor/Source/ActorAnimationEditor/Private/ActorAnimationEditorToolkit.cpp b/Engine/Plugins/MovieScene/ActorAnimationEditor/Source/ActorAnimationEditor/Private/ActorAnimationEditorToolkit.cpp
--- a/Engine/Plugins/MovieScene/ActorAnimationEditor/Source/ActorAnimationEditor/Private/ActorAnimationEditorToolkit.cpp
+++ b/Engine/Plugins/MovieScene/ActorAnimationEditor/Source/ActorAnimationEditor/Private/ActorAnimationEditorToolkit.cpp
@@ -112,6 +112,7 @@ void FActorAnimationEditorToolkit::Initialize( const EToolkitMode::Type Mode, co
 		);
 
 		SequencerInitParams.ViewParams.AddMenuExtender = AddMenuExtender;
+		SequencerInitParams.ViewParams.UniqueName = "ActorAnimationEditor";
 	}
 
 	Sequencer = FModuleManager::LoadModuleChecked<ISequencerModule>("Sequencer").CreateSequencer(SequencerInitParams);
